[NNE] Add a CoreML based runtime on mac.

Problem:
* We dont have a runtime that can run using apple in house inference engine and profit from all hardware feature there.

Solution:
* Integrate CoreML as a NNE Runtime, this CL is the first step only adding CPU inference (and thus NNERuntimeCPU interface), more hardware interface will be added as follow up. The feature set will also be expanded later, for exemple atm only MultiArray of float are supported for the models.

Test:
Mac dev editor: all nne tests (including the coreml tests from followup CL) + style transfer
Mac dev packaged : all nne tests (including the coreml tests from followup CL) + style transfer
Win64 dev editor : all need tests debug

#rnx

diff --git a/Engine/Plugins/Experimental/NNERuntimeCoreML/NNERuntimeCoreML.uplugin b/Engine/Plugins/Experimental/NNERuntimeCoreML/NNERuntimeCoreML.uplugin
new file mode 100644
--- /dev/null
+++ b/Engine/Plugins/Experimental/NNERuntimeCoreML/NNERuntimeCoreML.uplugin
@@ -0,0 +1,30 @@
+{
+	"FileVersion": 3,
+	"Version": 1,
+	"VersionName": "0.1",
+	"FriendlyName": "NNERuntimeCoreML",
+	"Description": "CoreML backed runtime for the Neural Network Engine (NNE).",
+	"Category": "ML",
+	"CreatedBy": "Epic Games, Inc.",
+	"CreatedByURL": "https://epicgames.com",
+	"DocsURL": "https://dev.epicgames.com/community/learning/courses/e7w/unreal-engine-neural-network-engine-nne",
+	"MarketplaceURL": "",
+	"SupportURL": "https://forums.unrealengine.com/t/course-neural-network-engine-nne/1162628",
+	"EnabledByDefault": false,
+	"CanContainContent": false,
+	"IsBetaVersion": false,
+	"IsExperimentalVersion": true,
+	"Installed": false,
+	"Modules": [
+		{
+			"Name": "NNERuntimeCoreMLEditor",
+			"Type": "Editor",
+			"LoadingPhase": "Default"
+		},
+		{
+			"Name": "NNERuntimeCoreML",
+			"Type": "RuntimeAndProgram",
+			"LoadingPhase": "Default"
+		}
+	]
+}
\ No newline at end of file
diff --git a/Engine/Plugins/Experimental/NNERuntimeCoreML/Source/NNERuntimeCoreML/Internal/NNERuntimeCoreMLLog.h b/Engine/Plugins/Experimental/NNERuntimeCoreML/Source/NNERuntimeCoreML/Internal/NNERuntimeCoreMLLog.h
new file mode 100644
--- /dev/null
+++ b/Engine/Plugins/Experimental/NNERuntimeCoreML/Source/NNERuntimeCoreML/Internal/NNERuntimeCoreMLLog.h
@@ -0,0 +1,7 @@
+// Copyright Epic Games, Inc. All Rights Reserved.
+
+#pragma once
+
+#include "Logging/LogMacros.h"
+
+DECLARE_LOG_CATEGORY_EXTERN(LogNNERuntimeCoreML, Log, All);
\ No newline at end of file
diff --git a/Engine/Plugins/Experimental/NNERuntimeCoreML/Source/NNERuntimeCoreML/NNERuntimeCoreML.Build.cs b/Engine/Plugins/Experimental/NNERuntimeCoreML/Source/NNERuntimeCoreML/NNERuntimeCoreML.Build.cs
new file mode 100644
--- /dev/null
+++ b/Engine/Plugins/Experimental/NNERuntimeCoreML/Source/NNERuntimeCoreML/NNERuntimeCoreML.Build.cs
@@ -0,0 +1,29 @@
+// Copyright Epic Games, Inc. All Rights Reserved.
+
+using UnrealBuildTool;
+
+public class NNERuntimeCoreML : ModuleRules
+{
+	public NNERuntimeCoreML( ReadOnlyTargetRules Target ) : base( Target )
+	{
+		PCHUsage = PCHUsageMode.UseExplicitOrSharedPCHs;
+
+		PrivateDependencyModuleNames.AddRange(new string[]
+		{
+			"Core",
+			"CoreUObject",
+			"NNE",
+			"Projects"
+		});
+
+		if (Target.Platform == UnrealTargetPlatform.Mac || Target.Type == TargetType.Editor)
+		{
+			PublicDefinitions.Add("WITH_NNE_RUNTIME_COREML");
+		}
+		
+		if (Target.Platform == UnrealTargetPlatform.Mac)
+		{
+			PublicFrameworks.AddRange(new string[] { "CoreML" });
+		}
+	}
+}
diff --git a/Engine/Plugins/Experimental/NNERuntimeCoreML/Source/NNERuntimeCoreML/Private/NNERuntimeCoreML.cpp b/Engine/Plugins/Experimental/NNERuntimeCoreML/Source/NNERuntimeCoreML/Private/NNERuntimeCoreML.cpp
new file mode 100644
--- /dev/null
+++ b/Engine/Plugins/Experimental/NNERuntimeCoreML/Source/NNERuntimeCoreML/Private/NNERuntimeCoreML.cpp
@@ -0,0 +1,98 @@
+// Copyright Epic Games, Inc. All Rights Reserved.
+
+#ifdef WITH_NNE_RUNTIME_COREML
+
+#include "NNERuntimeCoreML.h"
+
+#include "CoreMinimal.h"
+#include "NNE.h"
+#include "NNEModelData.h"
+#include "NNERuntimeCoreMLModel.h"
+#include "Serialization/MemoryWriter.h"
+
+FGuid UNNERuntimeCoreML::GUID = FGuid((int32)'C', (int32)'O', (int32)'M', (int32)'L');
+int32 UNNERuntimeCoreML::Version = 0x00000001;
+
+UNNERuntimeCoreML::ECanCreateModelDataStatus UNNERuntimeCoreML::CanCreateModelData(const FString& FileType, TConstArrayView64<uint8> FileData, const TMap<FString, TConstArrayView64<uint8>>& AdditionalFileData, const FGuid& FileId, const ITargetPlatform* TargetPlatform) const
+{
+	return (!FileData.IsEmpty() && FileType.Compare("mlmodel", ESearchCase::IgnoreCase) == 0) ? ECanCreateModelDataStatus::Ok : ECanCreateModelDataStatus::FailFileIdNotSupported;
+}
+
+TSharedPtr<UE::NNE::FSharedModelData> UNNERuntimeCoreML::CreateModelData(const FString& FileType, TConstArrayView64<uint8> FileData, const TMap<FString, TConstArrayView64<uint8>>& AdditionalFileData, const FGuid& FileId, const ITargetPlatform* TargetPlatform)
+{
+	using namespace UE::NNERuntimeCoreML::Private;
+
+	if (CanCreateModelData(FileType, FileData, AdditionalFileData, FileId, TargetPlatform) != ECanCreateModelDataStatus::Ok)
+	{
+		UE_LOG(LogNNERuntimeCoreML, Error, TEXT("Cannot create the CPU model data with id %s (Filetype: %s)"), *FileId.ToString(EGuidFormats::Digits).ToLower(), *FileType);
+		return {};
+	}
+
+	//Note: here the model should be optimized and the related mlmodelc binary blob saved instead of source model
+
+	TArray64<uint8> Result;
+	FMemoryWriter64 Writer(Result, /*bIsPersitent =*/ true);
+	Writer << UNNERuntimeCoreML::GUID;
+	Writer << UNNERuntimeCoreML::Version;
+
+	Writer.Serialize(const_cast<uint8*>(FileData.GetData()), FileData.Num());
+
+	return MakeShared<UE::NNE::FSharedModelData>(MakeSharedBufferFromArray(MoveTemp(Result)), 0);
+}
+
+FString UNNERuntimeCoreML::GetModelDataIdentifier(const FString& FileType, TConstArrayView64<uint8> FileData, const TMap<FString, TConstArrayView64<uint8>>& AdditionalFileData, const FGuid& FileId, const ITargetPlatform* TargetPlatform) const
+{
+	return FileId.ToString(EGuidFormats::Digits) + "-" + UNNERuntimeCoreML::GUID.ToString(EGuidFormats::Digits) + "-" + FString::FromInt(UNNERuntimeCoreML::Version);
+}
+
+FString UNNERuntimeCoreML::GetRuntimeName() const
+{
+	return TEXT("NNERuntimeCoreML");
+}
+
+UNNERuntimeCoreML::ECanCreateModelCPUStatus UNNERuntimeCoreML::CanCreateModelCPU(const TObjectPtr<UNNEModelData> ModelData) const
+{
+#ifdef WITH_NNE_RUNTIME_COREML
+	check(ModelData != nullptr);
+
+	constexpr int32 GuidSize = sizeof(UNNERuntimeCoreML::GUID);
+	constexpr int32 VersionSize = sizeof(UNNERuntimeCoreML::Version);
+	const TSharedPtr<UE::NNE::FSharedModelData> SharedData = ModelData->GetModelData(GetRuntimeName());
+
+	if (!SharedData.IsValid())
+	{
+		return ECanCreateModelCPUStatus::Fail;
+	}
+
+	TConstArrayView64<uint8> Data = SharedData->GetView();
+
+	if (Data.Num() <= GuidSize + VersionSize)
+	{
+		return ECanCreateModelCPUStatus::Fail;
+	}
+
+	bool bResult = FGenericPlatformMemory::Memcmp(&(Data[0]), &(UNNERuntimeCoreML::GUID), GuidSize) == 0;
+	bResult &= FGenericPlatformMemory::Memcmp(&(Data[GuidSize]), &(UNNERuntimeCoreML::Version), VersionSize) == 0;
+
+	return bResult ? ECanCreateModelCPUStatus::Ok : ECanCreateModelCPUStatus::Fail;
+#else
+	return ECanCreateModelCPUStatus::Fail;
+#endif 
+}
+
+TSharedPtr<UE::NNE::IModelCPU> UNNERuntimeCoreML::CreateModelCPU(const TObjectPtr<UNNEModelData> ModelData)
+{
+	check(ModelData != nullptr);
+
+	if (CanCreateModelCPU(ModelData) != ECanCreateModelCPUStatus::Ok)
+	{
+		UE_LOG(LogNNERuntimeCoreML, Error, TEXT("Cannot create a CPU model from the model data with id %s"), *ModelData->GetFileId().ToString(EGuidFormats::Digits));
+		return TSharedPtr<UE::NNE::IModelCPU>();
+	}
+
+	const TSharedRef<UE::NNE::FSharedModelData> SharedData = ModelData->GetModelData(GetRuntimeName()).ToSharedRef();
+
+	return MakeShared<UE::NNERuntimeCoreML::Private::FModelCoreMLCpu>(SharedData);
+}
+
+#endif // WITH_NNE_RUNTIME_COREML
diff --git a/Engine/Plugins/Experimental/NNERuntimeCoreML/Source/NNERuntimeCoreML/Private/NNERuntimeCoreML.h b/Engine/Plugins/Experimental/NNERuntimeCoreML/Source/NNERuntimeCoreML/Private/NNERuntimeCoreML.h
new file mode 100644
--- /dev/null
+++ b/Engine/Plugins/Experimental/NNERuntimeCoreML/Source/NNERuntimeCoreML/Private/NNERuntimeCoreML.h
@@ -0,0 +1,51 @@
+// Copyright Epic Games, Inc. All Rights Reserved.
+
+#pragma once
+
+#include "Logging/LogMacros.h"
+#include "Misc/Guid.h"
+#include "NNERuntime.h"
+#include "NNERuntimeCPU.h"
+#include "UObject/Object.h"
+
+#include "NNERuntimeCoreML.generated.h"
+
+DECLARE_LOG_CATEGORY_EXTERN(LogNNERuntimeCoreML, Log, All);
+
+UCLASS()
+class UNNERuntimeCoreML : public UObject, public INNERuntime, public INNERuntimeCPU
+{
+	GENERATED_BODY()
+
+public:
+	virtual ~UNNERuntimeCoreML() = default;
+
+#ifdef WITH_NNE_RUNTIME_COREML
+	static FGuid GUID;
+	static int32 Version;
+
+	//~ Begin INNERuntime Interface
+	virtual FString GetRuntimeName() const override;
+	virtual ECanCreateModelDataStatus CanCreateModelData(const FString& FileType, TConstArrayView64<uint8> FileData, const TMap<FString, TConstArrayView64<uint8>>& AdditionalFileData, const FGuid& FileId, const ITargetPlatform* TargetPlatform) const override;
+	virtual TSharedPtr<UE::NNE::FSharedModelData> CreateModelData(const FString& FileType, TConstArrayView64<uint8> FileData, const TMap<FString, TConstArrayView64<uint8>>& AdditionalFileData, const FGuid& FileId, const ITargetPlatform* TargetPlatform) override;
+	virtual FString GetModelDataIdentifier(const FString& FileType, TConstArrayView64<uint8> FileData, const TMap<FString, TConstArrayView64<uint8>>& AdditionalFileData, const FGuid& FileId, const ITargetPlatform* TargetPlatform) const override;
+	//~ End INNERuntime Interface
+
+	//~ Begin INNERuntimeCPU Interface
+	virtual ECanCreateModelCPUStatus CanCreateModelCPU(const TObjectPtr<UNNEModelData> ModelData) const override;
+	virtual TSharedPtr<UE::NNE::IModelCPU> CreateModelCPU(const TObjectPtr<UNNEModelData> ModelData) override;
+	//~ End INNERuntimeCPU Interface
+#else
+	//~ Begin INNERuntime Interface
+	virtual FString GetRuntimeName() const override { return ""; };
+	virtual ECanCreateModelDataStatus CanCreateModelData(const FString& FileType, TConstArrayView64<uint8> FileData, const TMap<FString, TConstArrayView64<uint8>>& AdditionalFileData, const FGuid& FileId, const ITargetPlatform* TargetPlatform) const override { return ECanCreateModelDataStatus::Fail; };
+	virtual TSharedPtr<UE::NNE::FSharedModelData> CreateModelData(const FString& FileType, TConstArrayView64<uint8> FileData, const TMap<FString, TConstArrayView64<uint8>>& AdditionalFileData, const FGuid& FileId, const ITargetPlatform* TargetPlatform) override { return TSharedPtr<UE::NNE::FSharedModelData>(); };
+	virtual FString GetModelDataIdentifier(const FString& FileType, TConstArrayView64<uint8> FileData, const TMap<FString, TConstArrayView64<uint8>>& AdditionalFileData, const FGuid& FileId, const ITargetPlatform* TargetPlatform) const override { return ""; };
+	//~ End INNERuntime Interface
+
+	//~ Begin INNERuntimeCPU Interface
+	virtual ECanCreateModelCPUStatus CanCreateModelCPU(const TObjectPtr<UNNEModelData> ModelData) const override { return ECanCreateModelCPUStatus::Fail; };
+	virtual TSharedPtr<UE::NNE::IModelCPU> CreateModelCPU(const TObjectPtr<UNNEModelData> ModelData) override { return TSharedPtr<UE::NNE::IModelCPU>(); };
+	//~ End INNERuntimeCPU Interface
+#endif // WITH_NNE_RUNTIME_COREML
+};
diff --git a/Engine/Plugins/Experimental/NNERuntimeCoreML/Source/NNERuntimeCoreML/Private/NNERuntimeCoreMLModel.cpp b/Engine/Plugins/Experimental/NNERuntimeCoreML/Source/NNERuntimeCoreML/Private/NNERuntimeCoreMLModel.cpp
new file mode 100644
--- /dev/null
+++ b/Engine/Plugins/Experimental/NNERuntimeCoreML/Source/NNERuntimeCoreML/Private/NNERuntimeCoreMLModel.cpp
@@ -0,0 +1,450 @@
+// Copyright Epic Games, Inc. All Rights Reserved.
+
+#ifdef WITH_NNE_RUNTIME_COREML
+
+#include "NNERuntimeCoreMLModel.h"
+
+#include "CoreMinimal.h"
+#include "HAL/PlatformFileManager.h"
+#include "HAL/PlatformMisc.h"
+#include "Misc/FileHelper.h"
+#include "Misc/Paths.h"
+#include "NNEModelData.h"
+#include "NNERuntimeCoreML.h"
+#include "NNETensor.h"
+#include "NNETypes.h"
+
+#if defined(__APPLE__)
+#import <Foundation/Foundation.h>
+#import <CoreML/CoreML.h>
+
+@interface CoreMLInputWrapper : NSObject<MLFeatureProvider>
+@property (readwrite, nonatomic, retain) NSMutableDictionary<NSString*, MLFeatureValue*>* featureByNames;
+@end
+
+@implementation CoreMLInputWrapper
+- (NSSet<NSString*>*) featureNames {
+	NSSet* set = [NSSet setWithArray:_featureByNames.allKeys];
+	return set;
+}
+- (nullable MLFeatureValue*) featureValueForName: (NSString*) featureName {
+	MLFeatureValue* value = _featureByNames[featureName];
+	return value;
+}
+@end
+
+#endif // defined(__APPLE__)
+
+namespace UE::NNERuntimeCoreML::Private
+{
+
+namespace Detail
+{
+#if defined(__APPLE__)
+	FString CreateTempDirPath(const FString& BasePath)
+	{
+		FString UniqueDirName;
+		do
+		{
+			UniqueDirName = FPaths::Combine(BasePath, *FString::Printf(TEXT("CoreMLModel_%s"), *FGuid::NewGuid().ToString()));
+		} while (IFileManager::Get().DirectoryExists(*UniqueDirName));
+		
+		return UniqueDirName;
+	}
+
+	void FillNSArrayShapeAndStrideFromNNEShape(const NNE::FTensorShape& TensorShape, NSMutableArray<NSNumber*>* Shape, NSMutableArray<NSNumber*>* Strides)
+	{
+		TArray<int32, TInlineAllocator<NNE::FTensorShape::MaxRank>> StridesData;
+		int32 OngoingStride = 1;
+		const int32 Rank = TensorShape.Rank();
+		
+		check(Rank > 0);
+		StridesData.SetNum(Rank);
+		for (int32 i = Rank-1; i >= 0 ; --i)
+		{
+			int32 DimSize = static_cast<int32>(TensorShape.GetData()[i]);
+			
+			StridesData[i] = OngoingStride;
+			OngoingStride *= DimSize;
+		}
+		check(OngoingStride == TensorShape.Volume());
+		
+		[Shape removeAllObjects];
+		[Strides removeAllObjects];
+		for (int32 i = 0; i < Rank ; ++i)
+		{
+			int32 DimSize = static_cast<int32>(TensorShape.GetData()[i]);
+			int32 DimStride = StridesData[i];
+			
+			[Shape addObject:[NSNumber numberWithInt:DimSize]];
+			[Strides addObject:[NSNumber numberWithInt:DimStride]];
+		}
+	}
+
+
+	template<typename T> TArray<T, TInlineAllocator<NNE::FTensorShape::MaxRank>> GetShapeDataFromNSArray(const NSArray<NSNumber*>* Shape)
+	{
+		TArray<T, TInlineAllocator<NNE::FTensorShape::MaxRank>> Data;
+		for(NSNumber* Dim in Shape)
+		{
+			Data.Emplace(static_cast<T>(Dim.intValue));
+		}
+		return Data;
+	}
+
+	bool GetTensorDescAndFeatureNamesFromMLDescription(TArray<NNE::FTensorDesc>& TensorDescs,
+										TArray<FString>& FeatureNames,
+										const NSDictionary<NSString* ,MLFeatureDescription *>* FeatureDictionary, 
+										const bool AreUndefinedShapesAnError, 
+										const FString& Context)
+	{
+		TensorDescs.Reset();
+		FeatureNames.Reset();
+		bool AllTensorDescsValid = true;
+		
+		if (FeatureDictionary == nullptr)
+		{
+			UE_LOG(LogNNERuntimeCoreML, Error, TEXT("Invalid feature dictionary for model %s. Was nullptr expecting a valid dictionary."), *Context);
+			return false;
+		}
+
+		for(NSString* Key in FeatureDictionary)
+		{
+			MLFeatureDescription* FeatureDescription = FeatureDictionary[Key];
+			const FString FeatureName = FeatureDescription.name;
+			
+			if (FeatureDescription.optional)
+			{
+				UE_LOG(LogNNERuntimeCoreML, Error, TEXT("Optional feature are not supported but %s feature %s is optional."), *Context, *FeatureName);
+				TensorDescs.Reset();
+				FeatureNames.Reset();
+				return false;
+			}
+			if (FeatureDescription.type != MLFeatureTypeMultiArray)
+			{
+				UE_LOG(LogNNERuntimeCoreML, Error, TEXT("Plugin is experimental: Only MultiArray are supported at the moment but %s feature %s is not"), *Context, *FeatureName);
+				TensorDescs.Reset();
+				FeatureNames.Reset();
+				return false;
+			}
+
+			if (FeatureDescription.multiArrayConstraint.dataType!=MLMultiArrayDataTypeFloat32)
+			{
+				UE_LOG(LogNNERuntimeCoreML, Error, TEXT("Plugin is experimental: Only float32 data type are supported at the moment but %s feature %s is not"), *Context, *FeatureName);
+				TensorDescs.Reset();
+				FeatureNames.Reset();
+				return false;
+			}
+			
+			if (FeatureDescription.multiArrayConstraint.shape.count != 0)
+			{
+				const TArray<int32, TInlineAllocator<NNE::FTensorShape::MaxRank>> SymbolicShapeData = GetShapeDataFromNSArray<int32>(FeatureDescription.multiArrayConstraint.shape);
+				const NNE::FSymbolicTensorShape SymbolicShape = NNE::FSymbolicTensorShape::Make(SymbolicShapeData);
+				const NNE::FTensorDesc SymbolicTensor = NNE::FTensorDesc::Make(FeatureName, SymbolicShape, ENNETensorDataType::Float);
+				
+				TensorDescs.Emplace(SymbolicTensor);
+			}
+			else // undefined MultiArray shape
+			{
+				if (AreUndefinedShapesAnError)
+				{
+					UE_LOG(LogNNERuntimeCoreML, Error, TEXT("MultiArray features need to define a shape but %s feature %s does not"), *Context, *FeatureName);
+					TensorDescs.Reset();
+					FeatureNames.Reset();
+					return false;
+				}
+				else
+				{
+					AllTensorDescsValid = false;
+				}
+			}
+			
+			FeatureNames.Emplace(FeatureName);
+		}
+		
+		if (!AllTensorDescsValid)
+		{
+			//if some CoreML features could not be expressed as TensorDescs we dont expose any of the metadata (or indices won't match).
+			//however we still have registered there names so we will be able to run inference.
+			TensorDescs.Reset();
+		}
+		
+		return true;
+	}
+#endif // defined(__APPLE__)
+} // namespace Detail
+
+FModelInstanceCoreMLCpu::~FModelInstanceCoreMLCpu()
+{
+#if defined(__APPLE__)
+	[CoreMLModelInstance release];
+#endif
+}
+
+TConstArrayView<NNE::FTensorDesc> FModelInstanceCoreMLCpu::GetInputTensorDescs() const
+{
+	return InputSymbolicTensors;
+}
+
+TConstArrayView<NNE::FTensorDesc> FModelInstanceCoreMLCpu::GetOutputTensorDescs() const
+{
+	return OutputSymbolicTensors;
+}
+
+TConstArrayView<NNE::FTensorShape> FModelInstanceCoreMLCpu::GetInputTensorShapes() const
+{
+	return InputTensorShapes;
+}
+
+TConstArrayView<NNE::FTensorShape> FModelInstanceCoreMLCpu::GetOutputTensorShapes() const
+{
+	return OutputTensorShapes;
+}
+
+FModelInstanceCoreMLCpu::ESetInputTensorShapesStatus FModelInstanceCoreMLCpu::SetInputTensorShapes(TConstArrayView<NNE::FTensorShape> InInputShapes)
+{
+	InputTensorShapes.Reset(InInputShapes.Num());
+	OutputTensorShapes.Reset();
+
+	if (InInputShapes.Num() != InputSymbolicTensors.Num())
+	{
+		UE_LOG(LogNNERuntimeCoreML, Error, TEXT("Number of input shapes does not match number of input tensors"));
+		return FModelInstanceCoreMLCpu::ESetInputTensorShapesStatus::Fail;
+	}
+
+	for (int32 i = 0; i < InInputShapes.Num(); ++i)
+	{
+		const NNE::FTensorDesc SymbolicDesc = InputSymbolicTensors[i];
+		if (!InInputShapes[i].IsCompatibleWith(SymbolicDesc.GetShape()))
+		{
+			UE_LOG(LogNNERuntimeCoreML, Error, TEXT("Input shape does not match input tensor %s of index %d"), *SymbolicDesc.GetName(), i);
+			return FModelInstanceCoreMLCpu::ESetInputTensorShapesStatus::Fail;
+		}
+	}
+		
+	InputTensorShapes = InInputShapes;
+
+	// Setup output shapes only if all model symbolic output shapes are concretes, otherwise it will be set during inference.
+	for (NNE::FTensorDesc SymbolicTensorDesc : OutputSymbolicTensors)
+	{
+		if (SymbolicTensorDesc.GetShape().IsConcrete())
+		{
+			NNE::Internal::FTensor Tensor = NNE::Internal::FTensor::MakeFromSymbolicDesc(SymbolicTensorDesc);
+			OutputTensorShapes.Emplace(Tensor.GetShape());
+		}
+	}
+	if (OutputTensorShapes.Num() != OutputSymbolicTensors.Num())
+	{
+		OutputTensorShapes.Reset();
+	}
+
+	return FModelInstanceCoreMLCpu::ESetInputTensorShapesStatus::Ok;
+}
+
+FModelInstanceCoreMLCpu::ERunSyncStatus FModelInstanceCoreMLCpu::RunSync(TConstArrayView<NNE::FTensorBindingCPU> InInputBindings, TConstArrayView<NNE::FTensorBindingCPU> InOutputBindings)
+{
+#if defined(__APPLE__)
+	SCOPED_AUTORELEASE_POOL;
+	
+	if (InputTensorShapes.IsEmpty())
+	{
+		UE_LOG(LogNNERuntimeCoreML, Error, TEXT("Input shapes are not set, please call SetInputTensorShapes."));
+		return FModelInstanceCoreMLCpu::ERunSyncStatus::Fail;
+	}
+
+	check(InputSymbolicTensors.Num() == InputTensorShapes.Num());
+
+	if (InInputBindings.Num() != InputTensorShapes.Num())
+	{
+		UE_LOG(LogNNERuntimeCoreML, Error, TEXT("Input bindings need to match input tensor descriptor count (got %d, expected %d)."), InInputBindings.Num(), InputTensorShapes.Num());
+		return FModelInstanceCoreMLCpu::ERunSyncStatus::Fail;
+	}
+
+	if (!InOutputBindings.IsEmpty() && InOutputBindings.Num() != OutputFeatureNames.Num())
+	{
+		UE_LOG(LogNNERuntimeCoreML, Error, TEXT("Output binding can be empty or needs to match output feature count (got %d, expected %d)."), InOutputBindings.Num(), OutputFeatureNames.Num());
+		return FModelInstanceCoreMLCpu::ERunSyncStatus::Fail;
+	}
+	
+	NSMutableDictionary<NSString*, MLFeatureValue*>* InputFeatureValueByNames = [NSMutableDictionary dictionaryWithCapacity:InInputBindings.Num()];
+	
+	for(int i = 0; i < InputFeatureNames.Num(); ++i)
+	{
+		NSError* ErrorCreatingInput = nullptr;
+		NNE::FTensorShape& InputTensorShape = InputTensorShapes[i];
+		NSMutableArray* Shape = [NSMutableArray arrayWithCapacity:InputTensorShape.Rank()];
+		NSMutableArray* Strides = [NSMutableArray arrayWithCapacity:InputTensorShape.Rank()];
+		
+		Detail::FillNSArrayShapeAndStrideFromNNEShape(InputTensorShape, Shape, Strides);
+
+		MLMultiArray* InputMultiArray = [[[MLMultiArray alloc]  initWithDataPointer:(void *) InInputBindings[i].Data 
+																shape:Shape
+																dataType:MLMultiArrayDataTypeFloat 
+																strides:Strides 
+																deallocator:^(void * _Nonnull bytes) {}
+																error: &ErrorCreatingInput]
+															autorelease];
+		
+		if (ErrorCreatingInput != nullptr)
+		{
+			FString Error = [ErrorCreatingInput localizedDescription];
+			UE_LOG(LogNNERuntimeCoreML, Error, TEXT("Failed to create the input (%s)"), *Error);
+			return FModelInstanceCoreMLCpu::ERunSyncStatus::Fail;
+		}
+		
+		[InputFeatureValueByNames setObject:[MLFeatureValue featureValueWithMultiArray:InputMultiArray] forKey:[NSString stringWithFString:InputFeatureNames[i]]];
+	}
+	
+	CoreMLInputWrapper* InputFeatureProvider = [[[CoreMLInputWrapper alloc] init] autorelease];
+	InputFeatureProvider.featureByNames = InputFeatureValueByNames;
+	
+	// Running inference
+	NSError* ErrorRunningInference = nullptr;
+	MLPredictionOptions* PredictionOptions = [[[MLPredictionOptions alloc] init] autorelease];
+	
+	id<MLFeatureProvider> OutFeatureProvider = [CoreMLModelInstance predictionFromFeatures:InputFeatureProvider options:PredictionOptions error:&ErrorRunningInference];
+	
+	if (ErrorRunningInference != nullptr)
+	{
+		FString Error = [ErrorRunningInference localizedDescription];
+		UE_LOG(LogNNERuntimeCoreML, Error, TEXT("Failed to run inference (%s)"), *Error);
+		return FModelInstanceCoreMLCpu::ERunSyncStatus::Fail;
+	}
+	if (!OutFeatureProvider) 
+	{
+		UE_LOG(LogNNERuntimeCoreML, Error, TEXT("No output feature provider was returned"));
+		return FModelInstanceCoreMLCpu::ERunSyncStatus::Fail;
+	}
+	
+	OutputTensorShapes.Reset();
+	for(int i = 0; i < OutputFeatureNames.Num(); ++i)
+	{
+		NSString* FeatureName = [NSString stringWithFString:OutputFeatureNames[i]];
+		MLMultiArray* OutputMultiArray = [OutFeatureProvider featureValueForName:FeatureName].multiArrayValue;
+		
+		// Capture the shapes of the output feature
+		const TArray<uint32, TInlineAllocator<NNE::FTensorShape::MaxRank>> ShapeData = Detail::GetShapeDataFromNSArray<uint32>(OutputMultiArray.shape);
+		const NNE::FTensorShape Shape = NNE::FTensorShape::Make(ShapeData);
+		
+		OutputTensorShapes.Emplace(Shape);
+		
+		// Copy the memory buffers of the output features to user provided bindings
+		if (!InOutputBindings.IsEmpty() && InOutputBindings[i].Data)
+		{
+			void (^copyToOutput)(const void*, NSInteger) =
+				^(const void* bytes, NSInteger size) {
+					if (size > 0 && InOutputBindings[i].SizeInBytes >= size)
+					{
+						FMemory::Memcpy(InOutputBindings[i].Data, bytes, size);
+					}
+				};
+			
+			[OutputMultiArray getBytesWithHandler:copyToOutput];
+		}
+	}
+	
+	return FModelInstanceCoreMLCpu::ERunSyncStatus::Ok;
+#else // !defined(__APPLE__)
+	return FModelInstanceCoreMLCpu::ERunSyncStatus::Fail;
+#endif // defined(__APPLE__)
+}
+
+bool FModelInstanceCoreMLCpu::Init(TConstArrayView64<uint8> ModelData)
+{
+#if defined(__APPLE__)
+	SCOPED_AUTORELEASE_POOL;
+	
+	// Find model memory blob and save it to disk.
+	int32 GuidSize = sizeof(UNNERuntimeCoreML::GUID);
+	int32 VersionSize = sizeof(UNNERuntimeCoreML::Version);
+	int32 GuidAndVersionSize = GuidSize+VersionSize;
+	TConstArrayView<uint8> ModelBuffer = { &(ModelData.GetData()[GuidAndVersionSize]), ModelData.Num() - GuidAndVersionSize };
+	FString ProjIntermediateDir = FPaths::ConvertRelativePathToFull(FPaths::ProjectIntermediateDir());
+	FString TempDirForModel = Detail::CreateTempDirPath(ProjIntermediateDir);
+	FString Filepath = FPaths::Combine(TempDirForModel, TEXT("CoreMLModel.mlmodel"));
+
+	if (!FFileHelper::SaveArrayToFile(ModelBuffer, *Filepath))
+	{
+		IFileManager::Get().DeleteDirectory(*TempDirForModel, false, true);
+		UE_LOG(LogNNERuntimeCoreML, Error, TEXT("Could not write model to disk at path %s."), *Filepath);
+		return false;
+	}
+	
+	// Compile the model with CoreML and cleanup temp directory
+	NSURL* ModelURL = [NSURL fileURLWithPath:[NSString stringWithFString:Filepath]];
+	NSError* ErrorCompilationObj = nullptr;
+	NSURL* CompiledModelURL = [MLModel compileModelAtURL:ModelURL error:&ErrorCompilationObj];
+	if (ErrorCompilationObj != nullptr)
+	{
+		IFileManager::Get().DeleteDirectory(*TempDirForModel, false, true);
+		FString Error = [ErrorCompilationObj localizedDescription];
+		UE_LOG(LogNNERuntimeCoreML, Error, TEXT("Failed to compile model (%s)"), *Error);
+		return false;
+	}
+	if (!IFileManager::Get().DeleteDirectory(*TempDirForModel, false, true))
+	{
+		UE_LOG(LogNNERuntimeCoreML, Error, TEXT("Could not cleanup model temp directy %s."), *TempDirForModel);
+	}
+
+	// Create the model instance from the compiled model.
+	NSError* ErrorCreationObj = nullptr;
+	MLModelConfiguration* ModelConfiguration = [[[MLModelConfiguration alloc] init] autorelease];
+	
+	ModelConfiguration.computeUnits = MLComputeUnitsCPUOnly;
+	
+	MLModel* ModelInstance = [MLModel modelWithContentsOfURL:CompiledModelURL configuration:ModelConfiguration error:&ErrorCreationObj];
+	
+	if (ErrorCreationObj != nullptr)
+	{
+		FString Error = [ErrorCreationObj localizedDescription];
+		UE_LOG(LogNNERuntimeCoreML, Error, TEXT("Failed to make an instance from compiled model (%s)"), *Error);
+		return false;
+	}
+	
+	CoreMLModelInstance = ModelInstance;
+	
+	// Extract model metadataß
+	if (!Detail::GetTensorDescAndFeatureNamesFromMLDescription(
+		InputSymbolicTensors, InputFeatureNames, 
+		CoreMLModelInstance.modelDescription.inputDescriptionsByName, 
+		true, TEXT("inputs")))
+	{
+		return false;
+	}
+	if (!Detail::GetTensorDescAndFeatureNamesFromMLDescription(
+		OutputSymbolicTensors, OutputFeatureNames, 
+		CoreMLModelInstance.modelDescription.outputDescriptionsByName, 
+		false, TEXT("outputs")))
+	{
+		return false;
+	}
+
+	[CoreMLModelInstance retain];
+	return true;
+#else // !defined(__APPLE__)
+	return false;
+#endif // defined(__APPLE__)
+}
+
+TSharedPtr<NNE::IModelInstanceCPU> FModelCoreMLCpu::CreateModelInstanceCPU()
+{
+	TSharedPtr<FModelInstanceCoreMLCpu> ModelInstance = MakeShared<FModelInstanceCoreMLCpu>();
+	bool IsModelInitialized = ModelInstance->Init(ModelData->GetView());
+	
+	if (!IsModelInitialized)
+	{
+		UE_LOG(LogNNERuntimeCoreML, Error, TEXT("Cannot initialize model instance"));
+		return TSharedPtr<NNE::IModelInstanceCPU>();
+	}
+
+	return ModelInstance;
+}
+
+FModelCoreMLCpu::FModelCoreMLCpu(TSharedRef<UE::NNE::FSharedModelData> InModelData) :
+	ModelData(InModelData)
+{
+}
+	
+} // namespace UE::NNERuntimeCoreML::Private
+
+#endif // WITH_NNE_RUNTIME_COREML
diff --git a/Engine/Plugins/Experimental/NNERuntimeCoreML/Source/NNERuntimeCoreML/Private/NNERuntimeCoreMLModel.h b/Engine/Plugins/Experimental/NNERuntimeCoreML/Source/NNERuntimeCoreML/Private/NNERuntimeCoreMLModel.h
new file mode 100644
--- /dev/null
+++ b/Engine/Plugins/Experimental/NNERuntimeCoreML/Source/NNERuntimeCoreML/Private/NNERuntimeCoreMLModel.h
@@ -0,0 +1,68 @@
+// Copyright Epic Games, Inc. All Rights Reserved.
+
+#pragma once
+
+#ifdef WITH_NNE_RUNTIME_COREML
+
+#include "CoreMinimal.h"
+#include "NNERuntimeCPU.h"
+#include "NNETensor.h"
+#include "NNETypes.h"
+
+#if defined(__APPLE__)
+@class MLModel;
+#endif
+
+namespace UE::NNE
+{
+	class FSharedModelData;
+}
+
+namespace UE::NNERuntimeCoreML::Private
+{
+
+class FModelInstanceCoreMLCpu : public NNE::IModelInstanceCPU
+{
+public:
+	FModelInstanceCoreMLCpu() = default;
+	virtual ~FModelInstanceCoreMLCpu();
+
+	using ESetInputTensorShapesStatus = UE::NNE::IModelInstanceCPU::ESetInputTensorShapesStatus;
+	using ERunSyncStatus = UE::NNE::IModelInstanceCPU::ERunSyncStatus;
+	
+	bool Init(TConstArrayView64<uint8> ModelData);
+
+	virtual TConstArrayView<NNE::FTensorDesc> GetInputTensorDescs() const override;
+	virtual TConstArrayView<NNE::FTensorDesc> GetOutputTensorDescs() const override;
+	virtual TConstArrayView<NNE::FTensorShape> GetInputTensorShapes() const override;
+	virtual TConstArrayView<NNE::FTensorShape> GetOutputTensorShapes() const override;
+	virtual ESetInputTensorShapesStatus SetInputTensorShapes(TConstArrayView<NNE::FTensorShape> InInputShapes) override;
+	virtual ERunSyncStatus RunSync(TConstArrayView<NNE::FTensorBindingCPU> InInputTensors, TConstArrayView<NNE::FTensorBindingCPU> InOutputTensors) override;
+
+private:
+	TArray<NNE::FTensorShape>	InputTensorShapes;
+	TArray<NNE::FTensorShape>	OutputTensorShapes;
+	TArray<NNE::FTensorDesc>	InputSymbolicTensors;
+	TArray<NNE::FTensorDesc>	OutputSymbolicTensors;
+#if defined(__APPLE__)
+	TArray<FString>				InputFeatureNames;
+	TArray<FString>				OutputFeatureNames;
+	MLModel *					CoreMLModelInstance;
+#endif
+};
+
+class FModelCoreMLCpu : public NNE::IModelCPU
+{
+public:
+	FModelCoreMLCpu(TSharedRef<NNE::FSharedModelData> InModelData);
+	virtual ~FModelCoreMLCpu() = default;
+
+	virtual TSharedPtr<NNE::IModelInstanceCPU> CreateModelInstanceCPU() override;
+
+private:
+	TSharedRef<NNE::FSharedModelData> ModelData;
+};
+	
+} // UE::NNERuntimeCoreML::Private
+
+#endif // WITH_NNE_RUNTIME_COREML
diff --git a/Engine/Plugins/Experimental/NNERuntimeCoreML/Source/NNERuntimeCoreML/Private/NNERuntimeCoreMLModule.cpp b/Engine/Plugins/Experimental/NNERuntimeCoreML/Source/NNERuntimeCoreML/Private/NNERuntimeCoreMLModule.cpp
new file mode 100644
--- /dev/null
+++ b/Engine/Plugins/Experimental/NNERuntimeCoreML/Source/NNERuntimeCoreML/Private/NNERuntimeCoreMLModule.cpp
@@ -0,0 +1,37 @@
+// Copyright Epic Games, Inc. All Rights Reserved.
+
+#include "NNERuntimeCoreMLModule.h"
+#include "NNE.h"
+#include "NNERuntimeCoreML.h"
+
+DEFINE_LOG_CATEGORY(LogNNERuntimeCoreML);
+
+void FNNERuntimeCoreMLModule::StartupModule()
+{
+#ifdef WITH_NNE_RUNTIME_COREML
+	NNERuntimeCoreML = NewObject<UNNERuntimeCoreML>();
+	if (NNERuntimeCoreML.IsValid())
+	{
+		TWeakInterfacePtr<INNERuntime> RuntimeInterface(NNERuntimeCoreML.Get());
+
+		NNERuntimeCoreML->AddToRoot();
+		UE::NNE::RegisterRuntime(RuntimeInterface);
+	}
+#endif // WITH_NNE_RUNTIME_COREML
+}
+
+void FNNERuntimeCoreMLModule::ShutdownModule()
+{
+#ifdef WITH_NNE_RUNTIME_COREML
+	if (NNERuntimeCoreML.IsValid())
+	{
+		TWeakInterfacePtr<INNERuntime> RuntimeInterface(NNERuntimeCoreML.Get());
+
+		UE::NNE::UnregisterRuntime(RuntimeInterface);
+		NNERuntimeCoreML->RemoveFromRoot();
+		NNERuntimeCoreML.Reset();
+	}
+#endif // WITH_NNE_RUNTIME_COREML
+}
+
+IMPLEMENT_MODULE(FNNERuntimeCoreMLModule, NNERuntimeCoreML);
diff --git a/Engine/Plugins/Experimental/NNERuntimeCoreML/Source/NNERuntimeCoreML/Private/NNERuntimeCoreMLModule.h b/Engine/Plugins/Experimental/NNERuntimeCoreML/Source/NNERuntimeCoreML/Private/NNERuntimeCoreMLModule.h
new file mode 100644
--- /dev/null
+++ b/Engine/Plugins/Experimental/NNERuntimeCoreML/Source/NNERuntimeCoreML/Private/NNERuntimeCoreMLModule.h
@@ -0,0 +1,23 @@
+// Copyright Epic Games, Inc. All Rights Reserved.
+
+#pragma once
+
+#include "Modules/ModuleManager.h"
+
+#ifdef WITH_NNE_RUNTIME_COREML
+#include "UObject/WeakObjectPtr.h"
+#include "UObject/WeakObjectPtrTemplates.h"
+
+class UNNERuntimeCoreML;
+#endif // WITH_NNE_RUNTIME_COREML
+
+class FNNERuntimeCoreMLModule : public IModuleInterface
+{
+private:
+#ifdef WITH_NNE_RUNTIME_COREML
+	TWeakObjectPtr<UNNERuntimeCoreML> NNERuntimeCoreML{ nullptr };
+#endif // WITH_NNE_RUNTIME_COREML
+
+	virtual void StartupModule() override;
+	virtual void ShutdownModule() override;
+};
diff --git a/Engine/Plugins/Experimental/NNERuntimeCoreML/Source/NNERuntimeCoreMLEditor/NNERuntimeCoreMLEditor.Build.cs b/Engine/Plugins/Experimental/NNERuntimeCoreML/Source/NNERuntimeCoreMLEditor/NNERuntimeCoreMLEditor.Build.cs
new file mode 100644
--- /dev/null
+++ b/Engine/Plugins/Experimental/NNERuntimeCoreML/Source/NNERuntimeCoreMLEditor/NNERuntimeCoreMLEditor.Build.cs
@@ -0,0 +1,27 @@
+// Copyright Epic Games, Inc. All Rights Reserved.
+
+using UnrealBuildTool;
+
+public class NNERuntimeCoreMLEditor : ModuleRules
+{
+	public NNERuntimeCoreMLEditor(ReadOnlyTargetRules Target) : base(Target)
+	{
+		PCHUsage = PCHUsageMode.UseExplicitOrSharedPCHs;
+
+		PrivateDependencyModuleNames.AddRange
+		(
+			new string[]
+			{
+				"Core",
+				"CoreUObject",
+				"Engine",
+				"MainFrame",
+				"NNE",
+				"NNERuntimeCoreML",
+				"Slate",
+				"SlateCore",
+				"UnrealEd"
+			}
+		);
+	}
+}
diff --git a/Engine/Plugins/Experimental/NNERuntimeCoreML/Source/NNERuntimeCoreMLEditor/Private/NNERuntimeCoreMLEditorModule.cpp b/Engine/Plugins/Experimental/NNERuntimeCoreML/Source/NNERuntimeCoreMLEditor/Private/NNERuntimeCoreMLEditorModule.cpp
new file mode 100644
--- /dev/null
+++ b/Engine/Plugins/Experimental/NNERuntimeCoreML/Source/NNERuntimeCoreMLEditor/Private/NNERuntimeCoreMLEditorModule.cpp
@@ -0,0 +1,18 @@
+// Copyright Epic Games, Inc. All Rights Reserved.
+
+#include "NNERuntimeCoreMLEditorModule.h"
+#include "NNERuntimeCoreMLLog.h"
+
+#define LOCTEXT_NAMESPACE "FNNERuntimeCoreMLEditorModule"
+
+void FNNERuntimeCoreMLEditorModule::StartupModule()
+{
+}
+
+void FNNERuntimeCoreMLEditorModule::ShutdownModule()
+{
+}
+
+#undef LOCTEXT_NAMESPACE
+	
+IMPLEMENT_MODULE(FNNERuntimeCoreMLEditorModule, NNERuntimeCoreMLEditor)
diff --git a/Engine/Plugins/Experimental/NNERuntimeCoreML/Source/NNERuntimeCoreMLEditor/Private/NNERuntimeCoreMLEditorModule.h b/Engine/Plugins/Experimental/NNERuntimeCoreML/Source/NNERuntimeCoreMLEditor/Private/NNERuntimeCoreMLEditorModule.h
new file mode 100644
--- /dev/null
+++ b/Engine/Plugins/Experimental/NNERuntimeCoreML/Source/NNERuntimeCoreMLEditor/Private/NNERuntimeCoreMLEditorModule.h
@@ -0,0 +1,13 @@
+// Copyright Epic Games, Inc. All Rights Reserved.
+
+#pragma once
+
+#include "CoreMinimal.h"
+#include "Modules/ModuleManager.h"
+
+class FNNERuntimeCoreMLEditorModule : public IModuleInterface
+{
+public:
+	virtual void StartupModule() override;
+	virtual void ShutdownModule() override;
+};
diff --git a/Engine/Plugins/Experimental/NNERuntimeCoreML/Source/NNERuntimeCoreMLEditor/Private/NNERuntimeCoreMLModelDataFactory.cpp b/Engine/Plugins/Experimental/NNERuntimeCoreML/Source/NNERuntimeCoreMLEditor/Private/NNERuntimeCoreMLModelDataFactory.cpp
new file mode 100644
--- /dev/null
+++ b/Engine/Plugins/Experimental/NNERuntimeCoreML/Source/NNERuntimeCoreMLEditor/Private/NNERuntimeCoreMLModelDataFactory.cpp
@@ -0,0 +1,44 @@
+// Copyright Epic Games, Inc. All Rights Reserved.
+
+#include "NNERuntimeCoreMLModelDataFactory.h"
+
+#include "Editor.h"
+#include "NNEModelData.h"
+#include "Misc/Paths.h"
+#include "Subsystems/ImportSubsystem.h"
+
+
+UNNERuntimeCoreMLModelDataFactory::UNNERuntimeCoreMLModelDataFactory(const FObjectInitializer& ObjectInitializer) : UFactory(ObjectInitializer)
+{
+	bCreateNew = false;
+	bEditorImport = true;
+	SupportedClass = UNNEModelData::StaticClass();
+	ImportPriority = DefaultImportPriority;
+	Formats.Add(TEXT("mlmodel;Core ML model format"));
+}
+
+UObject* UNNERuntimeCoreMLModelDataFactory::FactoryCreateBinary(UClass* InClass, UObject* InParent, FName Name, EObjectFlags Flags, UObject* Context, const TCHAR * Type, const uint8 *& Buffer, const uint8 * BufferEnd, FFeedbackContext* Warn)
+{
+	GEditor->GetEditorSubsystem<UImportSubsystem>()->BroadcastAssetPreImport(this, InClass, InParent, Name, Type);
+
+	if (!Type || !Buffer || !BufferEnd || BufferEnd - Buffer <= 0)
+	{
+		GEditor->GetEditorSubsystem<UImportSubsystem>()->BroadcastAssetPostImport(this, nullptr);
+		return nullptr;
+	}
+
+	TConstArrayView<uint8> BufferView = MakeArrayView(Buffer, BufferEnd - Buffer);
+	UNNEModelData* ModelData = NewObject<UNNEModelData>(InParent, InClass, Name, Flags);
+	check(ModelData)
+
+	ModelData->Init(Type, BufferView);
+
+	GEditor->GetEditorSubsystem<UImportSubsystem>()->BroadcastAssetPostImport(this, ModelData);
+
+	return ModelData;
+}
+
+bool UNNERuntimeCoreMLModelDataFactory::FactoryCanImport(const FString & Filename)
+{
+	return Filename.EndsWith(FString(TEXT("mlmodel")));
+}
diff --git a/Engine/Plugins/Experimental/NNERuntimeCoreML/Source/NNERuntimeCoreMLEditor/Public/NNERuntimeCoreMLModelDataFactory.h b/Engine/Plugins/Experimental/NNERuntimeCoreML/Source/NNERuntimeCoreMLEditor/Public/NNERuntimeCoreMLModelDataFactory.h
new file mode 100644
--- /dev/null
+++ b/Engine/Plugins/Experimental/NNERuntimeCoreML/Source/NNERuntimeCoreMLEditor/Public/NNERuntimeCoreMLModelDataFactory.h
@@ -0,0 +1,23 @@
+// Copyright Epic Games, Inc. All Rights Reserved.
+
+#pragma once
+
+#include "CoreMinimal.h"
+#include "Factories/Factory.h"
+
+#include "NNERuntimeCoreMLModelDataFactory.generated.h"
+
+UCLASS()
+class NNERUNTIMECOREMLEDITOR_API UNNERuntimeCoreMLModelDataFactory : public UFactory
+{
+	GENERATED_BODY()
+
+public:
+	UNNERuntimeCoreMLModelDataFactory(const FObjectInitializer& ObjectInitializer);
+
+public:
+	//~ Begin UFactory Interface
+	virtual UObject* FactoryCreateBinary(UClass* Class, UObject* InParent, FName Name, EObjectFlags Flags, UObject* Context, const TCHAR* Type, const uint8*& Buffer, const uint8* BufferEnd, FFeedbackContext* Warn) override;
+	virtual bool FactoryCanImport(const FString& Filename) override;
+	//~ End UFactory Interface
+};
\ No newline at end of file
