[NNE] Add GPU and NPU interface to CoreML runtime V2

Better version of https://p4-swarm.epicgames.net/reviews/39049810 (v1)

Problem: CoreML runtime was only supporting the CPU interface
Solution: Add GPU and NPU interfaces (NPU only on devices supportin it), Add NPU test to NNEQA

Tests:
Mac editor debug : all tests including verifying new interfaces are called for CoreML
Mac package debug: all tests including verifying new interfaces are called for CoreML
Windows editor debug : all tests
Windows package debug : all tests

#rnx

diff --git a/Engine/Plugins/Experimental/NNERuntimeCoreML/Source/NNERuntimeCoreML/Private/NNERuntimeCoreML.cpp b/Engine/Plugins/Experimental/NNERuntimeCoreML/Source/NNERuntimeCoreML/Private/NNERuntimeCoreML.cpp
--- a/Engine/Plugins/Experimental/NNERuntimeCoreML/Source/NNERuntimeCoreML/Private/NNERuntimeCoreML.cpp
+++ b/Engine/Plugins/Experimental/NNERuntimeCoreML/Source/NNERuntimeCoreML/Private/NNERuntimeCoreML.cpp
@@ -8,11 +8,41 @@
 #include "NNE.h"
 #include "NNEModelData.h"
 #include "NNERuntimeCoreMLModel.h"
+#include "NNERuntimeCoreMLNPUHelper.h"
 #include "Serialization/MemoryWriter.h"
 
 FGuid UNNERuntimeCoreML::GUID = FGuid((int32)'C', (int32)'O', (int32)'M', (int32)'L');
 int32 UNNERuntimeCoreML::Version = 0x00000001;
 
+namespace UE::NNERuntimeCoreML::Private::Details
+{
+	template <class CanCreateModelStatus> CanCreateModelStatus CanCreateModel(const TObjectPtr<UNNEModelData> ModelData, const FString& RuntimeName)
+	{
+		check(ModelData != nullptr);
+
+		constexpr int32 GuidSize = sizeof(UNNERuntimeCoreML::GUID);
+		constexpr int32 VersionSize = sizeof(UNNERuntimeCoreML::Version);
+		const TSharedPtr<UE::NNE::FSharedModelData> SharedData = ModelData->GetModelData(RuntimeName);
+
+		if (!SharedData.IsValid())
+		{
+			return CanCreateModelStatus::Fail;
+		}
+
+		TConstArrayView64<uint8> Data = SharedData->GetView();
+
+		if (Data.Num() <= GuidSize + VersionSize)
+		{
+			return CanCreateModelStatus::Fail;
+		}
+
+		bool bResult = FGenericPlatformMemory::Memcmp(&(Data[0]), &(UNNERuntimeCoreML::GUID), GuidSize) == 0;
+		bResult &= FGenericPlatformMemory::Memcmp(&(Data[GuidSize]), &(UNNERuntimeCoreML::Version), VersionSize) == 0;
+
+		return bResult ? CanCreateModelStatus::Ok : CanCreateModelStatus::Fail;
+	}
+} // UE::NNERuntimeCoreML::Private::Details
+
 UNNERuntimeCoreML::ECanCreateModelDataStatus UNNERuntimeCoreML::CanCreateModelData(const FString& FileType, TConstArrayView64<uint8> FileData, const TMap<FString, TConstArrayView64<uint8>>& AdditionalFileData, const FGuid& FileId, const ITargetPlatform* TargetPlatform) const
 {
 	return (!FileData.IsEmpty() && FileType.Compare("mlmodel", ESearchCase::IgnoreCase) == 0) ? ECanCreateModelDataStatus::Ok : ECanCreateModelDataStatus::FailFileIdNotSupported;
@@ -24,7 +54,7 @@ TSharedPtr<UE::NNE::FSharedModelData> UNNERuntimeCoreML::CreateModelData(const F
 
 	if (CanCreateModelData(FileType, FileData, AdditionalFileData, FileId, TargetPlatform) != ECanCreateModelDataStatus::Ok)
 	{
-		UE_LOG(LogNNERuntimeCoreML, Error, TEXT("Cannot create the CPU model data with id %s (Filetype: %s)"), *FileId.ToString(EGuidFormats::Digits).ToLower(), *FileType);
+		UE_LOG(LogNNERuntimeCoreML, Error, TEXT("Cannot create the model data with id %s (Filetype: %s)"), *FileId.ToString(EGuidFormats::Digits).ToLower(), *FileType);
 		return {};
 	}
 
@@ -50,49 +80,64 @@ FString UNNERuntimeCoreML::GetRuntimeName() const
 	return TEXT("NNERuntimeCoreML");
 }
 
-UNNERuntimeCoreML::ECanCreateModelCPUStatus UNNERuntimeCoreML::CanCreateModelCPU(const TObjectPtr<UNNEModelData> ModelData) const
+UNNERuntimeCoreMLCpuGpu::ECanCreateModelCPUStatus UNNERuntimeCoreMLCpuGpu::CanCreateModelCPU(const TObjectPtr<UNNEModelData> ModelData) const
 {
-#ifdef WITH_NNE_RUNTIME_COREML
-	check(ModelData != nullptr);
+	return UE::NNERuntimeCoreML::Private::Details::CanCreateModel<UNNERuntimeCoreMLCpuGpu::ECanCreateModelCPUStatus>(ModelData, GetRuntimeName());
+}
+
+UNNERuntimeCoreMLCpuGpu::ECanCreateModelGPUStatus UNNERuntimeCoreMLCpuGpu::CanCreateModelGPU(const TObjectPtr<UNNEModelData> ModelData) const
+{
+	return UE::NNERuntimeCoreML::Private::Details::CanCreateModel<UNNERuntimeCoreMLCpuGpu::ECanCreateModelGPUStatus>(ModelData, GetRuntimeName());
+}
+
+UNNERuntimeCoreMLCpuGpuNpu::ECanCreateModelNPUStatus UNNERuntimeCoreMLCpuGpuNpu::CanCreateModelNPU(const TObjectPtr<UNNEModelData> ModelData) const
+{
+	if (!UE::NNERuntimeCoreML::Private::IsNPUAvailable())
+		return UNNERuntimeCoreMLCpuGpuNpu::ECanCreateModelNPUStatus::Fail;
 
-	constexpr int32 GuidSize = sizeof(UNNERuntimeCoreML::GUID);
-	constexpr int32 VersionSize = sizeof(UNNERuntimeCoreML::Version);
-	const TSharedPtr<UE::NNE::FSharedModelData> SharedData = ModelData->GetModelData(GetRuntimeName());
+	return UE::NNERuntimeCoreML::Private::Details::CanCreateModel<UNNERuntimeCoreMLCpuGpuNpu::ECanCreateModelNPUStatus>(ModelData, GetRuntimeName());
+}
 
-	if (!SharedData.IsValid())
+TSharedPtr<UE::NNE::IModelCPU> UNNERuntimeCoreMLCpuGpu::CreateModelCPU(const TObjectPtr<UNNEModelData> ModelData)
+{
+	check(ModelData != nullptr);
+	if (CanCreateModelCPU(ModelData) != ECanCreateModelCPUStatus::Ok)
 	{
-		return ECanCreateModelCPUStatus::Fail;
+		UE_LOG(LogNNERuntimeCoreML, Error, TEXT("Cannot create a CPU model from the model data with id %s"), *ModelData->GetFileId().ToString(EGuidFormats::Digits));
+		return TSharedPtr<UE::NNE::IModelCPU>();
 	}
 
-	TConstArrayView64<uint8> Data = SharedData->GetView();
+	const TSharedRef<UE::NNE::FSharedModelData> SharedData = ModelData->GetModelData(GetRuntimeName()).ToSharedRef();
+
+	return MakeShared<UE::NNERuntimeCoreML::Private::FModelCoreMLCpu>(SharedData);
+}
 
-	if (Data.Num() <= GuidSize + VersionSize)
+TSharedPtr<UE::NNE::IModelGPU> UNNERuntimeCoreMLCpuGpu::CreateModelGPU(const TObjectPtr<UNNEModelData> ModelData)
+{
+	check(ModelData != nullptr);
+	if (CanCreateModelGPU(ModelData) != ECanCreateModelGPUStatus::Ok)
 	{
-		return ECanCreateModelCPUStatus::Fail;
+		UE_LOG(LogNNERuntimeCoreML, Error, TEXT("Cannot create a GPU model from the model data with id %s"), *ModelData->GetFileId().ToString(EGuidFormats::Digits));
+		return TSharedPtr<UE::NNE::IModelGPU>();
 	}
 
-	bool bResult = FGenericPlatformMemory::Memcmp(&(Data[0]), &(UNNERuntimeCoreML::GUID), GuidSize) == 0;
-	bResult &= FGenericPlatformMemory::Memcmp(&(Data[GuidSize]), &(UNNERuntimeCoreML::Version), VersionSize) == 0;
+	const TSharedRef<UE::NNE::FSharedModelData> SharedData = ModelData->GetModelData(GetRuntimeName()).ToSharedRef();
 
-	return bResult ? ECanCreateModelCPUStatus::Ok : ECanCreateModelCPUStatus::Fail;
-#else
-	return ECanCreateModelCPUStatus::Fail;
-#endif 
+	return MakeShared<UE::NNERuntimeCoreML::Private::FModelCoreMLGpu>(SharedData);
 }
 
-TSharedPtr<UE::NNE::IModelCPU> UNNERuntimeCoreML::CreateModelCPU(const TObjectPtr<UNNEModelData> ModelData)
+TSharedPtr<UE::NNE::IModelNPU> UNNERuntimeCoreMLCpuGpuNpu::CreateModelNPU(const TObjectPtr<UNNEModelData> ModelData)
 {
 	check(ModelData != nullptr);
-
-	if (CanCreateModelCPU(ModelData) != ECanCreateModelCPUStatus::Ok)
+	if (CanCreateModelNPU(ModelData) != ECanCreateModelNPUStatus::Ok)
 	{
-		UE_LOG(LogNNERuntimeCoreML, Error, TEXT("Cannot create a CPU model from the model data with id %s"), *ModelData->GetFileId().ToString(EGuidFormats::Digits));
-		return TSharedPtr<UE::NNE::IModelCPU>();
+		UE_LOG(LogNNERuntimeCoreML, Error, TEXT("Cannot create a NPU model from the model data with id %s"), *ModelData->GetFileId().ToString(EGuidFormats::Digits));
+		return TSharedPtr<UE::NNE::IModelNPU>();
 	}
 
 	const TSharedRef<UE::NNE::FSharedModelData> SharedData = ModelData->GetModelData(GetRuntimeName()).ToSharedRef();
 
-	return MakeShared<UE::NNERuntimeCoreML::Private::FModelCoreMLCpu>(SharedData);
+	return MakeShared<UE::NNERuntimeCoreML::Private::FModelCoreMLNpu>(SharedData);
 }
 
 #endif // WITH_NNE_RUNTIME_COREML
diff --git a/Engine/Plugins/Experimental/NNERuntimeCoreML/Source/NNERuntimeCoreML/Private/NNERuntimeCoreML.h b/Engine/Plugins/Experimental/NNERuntimeCoreML/Source/NNERuntimeCoreML/Private/NNERuntimeCoreML.h
--- a/Engine/Plugins/Experimental/NNERuntimeCoreML/Source/NNERuntimeCoreML/Private/NNERuntimeCoreML.h
+++ b/Engine/Plugins/Experimental/NNERuntimeCoreML/Source/NNERuntimeCoreML/Private/NNERuntimeCoreML.h
@@ -6,6 +6,8 @@
 #include "Misc/Guid.h"
 #include "NNERuntime.h"
 #include "NNERuntimeCPU.h"
+#include "NNERuntimeGPU.h"
+#include "NNERuntimeNPU.h"
 #include "UObject/Object.h"
 
 #include "NNERuntimeCoreML.generated.h"
@@ -13,7 +15,7 @@
 DECLARE_LOG_CATEGORY_EXTERN(LogNNERuntimeCoreML, Log, All);
 
 UCLASS()
-class UNNERuntimeCoreML : public UObject, public INNERuntime, public INNERuntimeCPU
+class UNNERuntimeCoreML : public UObject, public INNERuntime
 {
 	GENERATED_BODY()
 
@@ -30,11 +32,6 @@ public:
 	virtual TSharedPtr<UE::NNE::FSharedModelData> CreateModelData(const FString& FileType, TConstArrayView64<uint8> FileData, const TMap<FString, TConstArrayView64<uint8>>& AdditionalFileData, const FGuid& FileId, const ITargetPlatform* TargetPlatform) override;
 	virtual FString GetModelDataIdentifier(const FString& FileType, TConstArrayView64<uint8> FileData, const TMap<FString, TConstArrayView64<uint8>>& AdditionalFileData, const FGuid& FileId, const ITargetPlatform* TargetPlatform) const override;
 	//~ End INNERuntime Interface
-
-	//~ Begin INNERuntimeCPU Interface
-	virtual ECanCreateModelCPUStatus CanCreateModelCPU(const TObjectPtr<UNNEModelData> ModelData) const override;
-	virtual TSharedPtr<UE::NNE::IModelCPU> CreateModelCPU(const TObjectPtr<UNNEModelData> ModelData) override;
-	//~ End INNERuntimeCPU Interface
 #else
 	//~ Begin INNERuntime Interface
 	virtual FString GetRuntimeName() const override { return ""; };
@@ -42,10 +39,57 @@ public:
 	virtual TSharedPtr<UE::NNE::FSharedModelData> CreateModelData(const FString& FileType, TConstArrayView64<uint8> FileData, const TMap<FString, TConstArrayView64<uint8>>& AdditionalFileData, const FGuid& FileId, const ITargetPlatform* TargetPlatform) override { return TSharedPtr<UE::NNE::FSharedModelData>(); };
 	virtual FString GetModelDataIdentifier(const FString& FileType, TConstArrayView64<uint8> FileData, const TMap<FString, TConstArrayView64<uint8>>& AdditionalFileData, const FGuid& FileId, const ITargetPlatform* TargetPlatform) const override { return ""; };
 	//~ End INNERuntime Interface
+#endif // WITH_NNE_RUNTIME_COREML
+};
+
+UCLASS()
+class UNNERuntimeCoreMLCpuGpu : public UNNERuntimeCoreML, public INNERuntimeCPU, public INNERuntimeGPU
+{
+	GENERATED_BODY()
+
+public:
+	virtual ~UNNERuntimeCoreMLCpuGpu() = default;
+
+#ifdef WITH_NNE_RUNTIME_COREML
+	//~ Begin INNERuntimeCPU Interface
+	virtual ECanCreateModelCPUStatus CanCreateModelCPU(const TObjectPtr<UNNEModelData> ModelData) const override;
+	virtual TSharedPtr<UE::NNE::IModelCPU> CreateModelCPU(const TObjectPtr<UNNEModelData> ModelData) override;
+	//~ End INNERuntimeCPU Interface
 
+	//~ Begin INNERuntimeGPU Interface
+	virtual ECanCreateModelGPUStatus CanCreateModelGPU(const TObjectPtr<UNNEModelData> ModelData) const override;
+	virtual TSharedPtr<UE::NNE::IModelGPU> CreateModelGPU(const TObjectPtr<UNNEModelData> ModelData) override;
+	//~ End INNERuntimeGPU Interface
+#else
 	//~ Begin INNERuntimeCPU Interface
 	virtual ECanCreateModelCPUStatus CanCreateModelCPU(const TObjectPtr<UNNEModelData> ModelData) const override { return ECanCreateModelCPUStatus::Fail; };
 	virtual TSharedPtr<UE::NNE::IModelCPU> CreateModelCPU(const TObjectPtr<UNNEModelData> ModelData) override { return TSharedPtr<UE::NNE::IModelCPU>(); };
 	//~ End INNERuntimeCPU Interface
+
+	//~ Begin INNERuntimeGPU Interface
+	virtual ECanCreateModelGPUStatus CanCreateModelGPU(const TObjectPtr<UNNEModelData> ModelData) const override { return ECanCreateModelGPUStatus::Fail; };
+	virtual TSharedPtr<UE::NNE::IModelGPU> CreateModelGPU(const TObjectPtr<UNNEModelData> ModelData) override { return TSharedPtr<UE::NNE::IModelGPU>(); };
+	//~ End INNERuntimeGPU Interface
+#endif // WITH_NNE_RUNTIME_COREML
+};
+
+UCLASS()
+class UNNERuntimeCoreMLCpuGpuNpu : public UNNERuntimeCoreMLCpuGpu, public INNERuntimeNPU
+{
+	GENERATED_BODY()
+
+public:
+	virtual ~UNNERuntimeCoreMLCpuGpuNpu() = default;
+
+#ifdef WITH_NNE_RUNTIME_COREML
+	//~ Begin INNERuntimeNPU Interface
+	virtual ECanCreateModelNPUStatus CanCreateModelNPU(const TObjectPtr<UNNEModelData> ModelData) const override;
+	virtual TSharedPtr<UE::NNE::IModelNPU> CreateModelNPU(const TObjectPtr<UNNEModelData> ModelData) override;
+	//~ End INNERuntimeNPU Interface
+#else
+	//~ Begin INNERuntimeNPU Interface
+	virtual ECanCreateModelNPUStatus CanCreateModelNPU(const TObjectPtr<UNNEModelData> ModelData) const override { return ECanCreateModelNPUStatus::Fail; };
+	virtual TSharedPtr<UE::NNE::IModelNPU> CreateModelNPU(const TObjectPtr<UNNEModelData> ModelData) override { return TSharedPtr<UE::NNE::IModelNPU>(); };
+	//~ End INNERuntimeNPU Interface
 #endif // WITH_NNE_RUNTIME_COREML
 };
diff --git a/Engine/Plugins/Experimental/NNERuntimeCoreML/Source/NNERuntimeCoreML/Private/NNERuntimeCoreMLModel.cpp b/Engine/Plugins/Experimental/NNERuntimeCoreML/Source/NNERuntimeCoreML/Private/NNERuntimeCoreMLModel.cpp
--- a/Engine/Plugins/Experimental/NNERuntimeCoreML/Source/NNERuntimeCoreML/Private/NNERuntimeCoreMLModel.cpp
+++ b/Engine/Plugins/Experimental/NNERuntimeCoreML/Source/NNERuntimeCoreML/Private/NNERuntimeCoreMLModel.cpp
@@ -41,6 +41,26 @@ namespace UE::NNERuntimeCoreML::Private
 namespace Detail
 {
 #if defined(__APPLE__)
+
+	template< typename T > 
+	struct InstanceTypeTrait{ 
+	};
+
+	template<> 
+	struct InstanceTypeTrait<UE::NNE::IModelInstanceCPU>{ 
+		static MLComputeUnits GetComputeUnits() { return MLComputeUnitsCPUOnly; }
+	};
+
+	template<> 
+	struct InstanceTypeTrait<UE::NNE::IModelInstanceGPU>{ 
+		static MLComputeUnits GetComputeUnits() { return MLComputeUnitsCPUAndGPU; }
+	};
+
+	template<> 
+	struct InstanceTypeTrait<UE::NNE::IModelInstanceNPU>{ 
+		static MLComputeUnits GetComputeUnits() { return MLComputeUnitsCPUAndNeuralEngine; }
+	};
+
 	FString CreateTempDirPath(const FString& BasePath)
 	{
 		FString UniqueDirName;
@@ -174,34 +194,40 @@ namespace Detail
 #endif // defined(__APPLE__)
 } // namespace Detail
 
-FModelInstanceCoreMLCpu::~FModelInstanceCoreMLCpu()
+template <class InstanceType>
+FModelInstanceCoreMLBase<InstanceType>::~FModelInstanceCoreMLBase()
 {
 #if defined(__APPLE__)
 	[CoreMLModelInstance release];
-#endif
+#endif // defined(__APPLE__)
 }
 
-TConstArrayView<NNE::FTensorDesc> FModelInstanceCoreMLCpu::GetInputTensorDescs() const
+template <class InstanceType>
+TConstArrayView<NNE::FTensorDesc> FModelInstanceCoreMLBase<InstanceType>::GetInputTensorDescs() const
 {
 	return InputSymbolicTensors;
 }
 
-TConstArrayView<NNE::FTensorDesc> FModelInstanceCoreMLCpu::GetOutputTensorDescs() const
+template <class InstanceType>
+TConstArrayView<NNE::FTensorDesc> FModelInstanceCoreMLBase<InstanceType>::GetOutputTensorDescs() const
 {
 	return OutputSymbolicTensors;
 }
 
-TConstArrayView<NNE::FTensorShape> FModelInstanceCoreMLCpu::GetInputTensorShapes() const
+template <class InstanceType>
+TConstArrayView<NNE::FTensorShape> FModelInstanceCoreMLBase<InstanceType>::GetInputTensorShapes() const
 {
 	return InputTensorShapes;
 }
 
-TConstArrayView<NNE::FTensorShape> FModelInstanceCoreMLCpu::GetOutputTensorShapes() const
+template <class InstanceType>
+TConstArrayView<NNE::FTensorShape> FModelInstanceCoreMLBase<InstanceType>::GetOutputTensorShapes() const
 {
 	return OutputTensorShapes;
 }
 
-FModelInstanceCoreMLCpu::ESetInputTensorShapesStatus FModelInstanceCoreMLCpu::SetInputTensorShapes(TConstArrayView<NNE::FTensorShape> InInputShapes)
+template <class InstanceType>
+FModelInstanceCoreMLCpu::ESetInputTensorShapesStatus FModelInstanceCoreMLBase<InstanceType>::SetInputTensorShapes(TConstArrayView<NNE::FTensorShape> InInputShapes)
 {
 	InputTensorShapes.Reset(InInputShapes.Num());
 	OutputTensorShapes.Reset();
@@ -241,7 +267,8 @@ FModelInstanceCoreMLCpu::ESetInputTensorShapesStatus FModelInstanceCoreMLCpu::Se
 	return FModelInstanceCoreMLCpu::ESetInputTensorShapesStatus::Ok;
 }
 
-FModelInstanceCoreMLCpu::ERunSyncStatus FModelInstanceCoreMLCpu::RunSync(TConstArrayView<NNE::FTensorBindingCPU> InInputBindings, TConstArrayView<NNE::FTensorBindingCPU> InOutputBindings)
+template <class InstanceType>
+FModelInstanceCoreMLCpu::ERunSyncStatus FModelInstanceCoreMLBase<InstanceType>::RunSync(TConstArrayView<NNE::FTensorBindingCPU> InInputBindings, TConstArrayView<NNE::FTensorBindingCPU> InOutputBindings)
 {
 #if defined(__APPLE__)
 	SCOPED_AUTORELEASE_POOL;
@@ -349,7 +376,8 @@ FModelInstanceCoreMLCpu::ERunSyncStatus FModelInstanceCoreMLCpu::RunSync(TConstA
 #endif // defined(__APPLE__)
 }
 
-bool FModelInstanceCoreMLCpu::Init(TConstArrayView64<uint8> ModelData)
+template <class InstanceType>
+bool FModelInstanceCoreMLBase<InstanceType>::Init(TConstArrayView64<uint8> ModelData)
 {
 #if defined(__APPLE__)
 	SCOPED_AUTORELEASE_POOL;
@@ -390,7 +418,7 @@ bool FModelInstanceCoreMLCpu::Init(TConstArrayView64<uint8> ModelData)
 	NSError* ErrorCreationObj = nullptr;
 	MLModelConfiguration* ModelConfiguration = [[[MLModelConfiguration alloc] init] autorelease];
 	
-	ModelConfiguration.computeUnits = MLComputeUnitsCPUOnly;
+	ModelConfiguration.computeUnits = Detail::InstanceTypeTrait<InstanceType>::GetComputeUnits();
 	
 	MLModel* ModelInstance = [MLModel modelWithContentsOfURL:CompiledModelURL configuration:ModelConfiguration error:&ErrorCreationObj];
 	
@@ -440,10 +468,48 @@ TSharedPtr<NNE::IModelInstanceCPU> FModelCoreMLCpu::CreateModelInstanceCPU()
 	return ModelInstance;
 }
 
+TSharedPtr<NNE::IModelInstanceGPU> FModelCoreMLGpu::CreateModelInstanceGPU()
+{
+	TSharedPtr<FModelInstanceCoreMLGpu> ModelInstance = MakeShared<FModelInstanceCoreMLGpu>();
+	bool IsModelInitialized = ModelInstance->Init(ModelData->GetView());
+	
+	if (!IsModelInitialized)
+	{
+		UE_LOG(LogNNERuntimeCoreML, Error, TEXT("Cannot initialize model instance"));
+		return TSharedPtr<NNE::IModelInstanceGPU>();
+	}
+
+	return ModelInstance;
+}
+
+TSharedPtr<NNE::IModelInstanceNPU> FModelCoreMLNpu::CreateModelInstanceNPU()
+{
+	TSharedPtr<FModelInstanceCoreMLNpu> ModelInstance = MakeShared<FModelInstanceCoreMLNpu>();
+	bool IsModelInitialized = ModelInstance->Init(ModelData->GetView());
+	
+	if (!IsModelInitialized)
+	{
+		UE_LOG(LogNNERuntimeCoreML, Error, TEXT("Cannot initialize model instance"));
+		return TSharedPtr<NNE::IModelInstanceNPU>();
+	}
+
+	return ModelInstance;
+}
+
 FModelCoreMLCpu::FModelCoreMLCpu(TSharedRef<UE::NNE::FSharedModelData> InModelData) :
 	ModelData(InModelData)
 {
 }
+
+FModelCoreMLGpu::FModelCoreMLGpu(TSharedRef<UE::NNE::FSharedModelData> InModelData) :
+	ModelData(InModelData)
+{
+}
+
+FModelCoreMLNpu::FModelCoreMLNpu(TSharedRef<UE::NNE::FSharedModelData> InModelData) :
+	ModelData(InModelData)
+{
+}
 	
 } // namespace UE::NNERuntimeCoreML::Private
 
diff --git a/Engine/Plugins/Experimental/NNERuntimeCoreML/Source/NNERuntimeCoreML/Private/NNERuntimeCoreMLModel.h b/Engine/Plugins/Experimental/NNERuntimeCoreML/Source/NNERuntimeCoreML/Private/NNERuntimeCoreMLModel.h
--- a/Engine/Plugins/Experimental/NNERuntimeCoreML/Source/NNERuntimeCoreML/Private/NNERuntimeCoreMLModel.h
+++ b/Engine/Plugins/Experimental/NNERuntimeCoreML/Source/NNERuntimeCoreML/Private/NNERuntimeCoreMLModel.h
@@ -6,6 +6,8 @@
 
 #include "CoreMinimal.h"
 #include "NNERuntimeCPU.h"
+#include "NNERuntimeGPU.h"
+#include "NNERuntimeNPU.h"
 #include "NNETensor.h"
 #include "NNETypes.h"
 
@@ -21,11 +23,12 @@ namespace UE::NNE
 namespace UE::NNERuntimeCoreML::Private
 {
 
-class FModelInstanceCoreMLCpu : public NNE::IModelInstanceCPU
+template <class InstanceType>
+class FModelInstanceCoreMLBase : public InstanceType
 {
 public:
-	FModelInstanceCoreMLCpu() = default;
-	virtual ~FModelInstanceCoreMLCpu();
+	FModelInstanceCoreMLBase() = default;
+	virtual ~FModelInstanceCoreMLBase();
 
 	using ESetInputTensorShapesStatus = UE::NNE::IModelInstanceCPU::ESetInputTensorShapesStatus;
 	using ERunSyncStatus = UE::NNE::IModelInstanceCPU::ERunSyncStatus;
@@ -51,6 +54,27 @@ private:
 #endif
 };
 
+class FModelInstanceCoreMLCpu : public FModelInstanceCoreMLBase<NNE::IModelInstanceCPU>
+{
+public:
+	FModelInstanceCoreMLCpu() = default;
+	virtual ~FModelInstanceCoreMLCpu() = default;
+};
+
+class FModelInstanceCoreMLGpu : public FModelInstanceCoreMLBase<NNE::IModelInstanceGPU>
+{
+public:
+	FModelInstanceCoreMLGpu() = default;
+	virtual ~FModelInstanceCoreMLGpu() = default;
+};
+
+class FModelInstanceCoreMLNpu : public FModelInstanceCoreMLBase<NNE::IModelInstanceNPU>
+{
+public:
+	FModelInstanceCoreMLNpu() = default;
+	virtual ~FModelInstanceCoreMLNpu() = default;
+};
+
 class FModelCoreMLCpu : public NNE::IModelCPU
 {
 public:
@@ -62,6 +86,30 @@ public:
 private:
 	TSharedRef<NNE::FSharedModelData> ModelData;
 };
+
+class FModelCoreMLGpu : public NNE::IModelGPU
+{
+public:
+	FModelCoreMLGpu(TSharedRef<NNE::FSharedModelData> InModelData);
+	virtual ~FModelCoreMLGpu() = default;
+
+	virtual TSharedPtr<NNE::IModelInstanceGPU> CreateModelInstanceGPU() override;
+
+private:
+	TSharedRef<NNE::FSharedModelData> ModelData;
+};
+
+class FModelCoreMLNpu : public NNE::IModelNPU
+{
+public:
+	FModelCoreMLNpu(TSharedRef<NNE::FSharedModelData> InModelData);
+	virtual ~FModelCoreMLNpu() = default;
+
+	virtual TSharedPtr<NNE::IModelInstanceNPU> CreateModelInstanceNPU() override;
+
+private:
+	TSharedRef<NNE::FSharedModelData> ModelData;
+};
 	
 } // UE::NNERuntimeCoreML::Private
 
diff --git a/Engine/Plugins/Experimental/NNERuntimeCoreML/Source/NNERuntimeCoreML/Private/NNERuntimeCoreMLModule.cpp b/Engine/Plugins/Experimental/NNERuntimeCoreML/Source/NNERuntimeCoreML/Private/NNERuntimeCoreMLModule.cpp
--- a/Engine/Plugins/Experimental/NNERuntimeCoreML/Source/NNERuntimeCoreML/Private/NNERuntimeCoreMLModule.cpp
+++ b/Engine/Plugins/Experimental/NNERuntimeCoreML/Source/NNERuntimeCoreML/Private/NNERuntimeCoreMLModule.cpp
@@ -3,17 +3,35 @@
 #include "NNERuntimeCoreMLModule.h"
 #include "NNE.h"
 #include "NNERuntimeCoreML.h"
+#include "NNERuntimeCoreMLNPUHelper.h"
+
+#if defined(__APPLE__)
+#import <Foundation/Foundation.h>
+#import <CoreML/CoreML.h>
+#endif // defined(__APPLE__)
 
 DEFINE_LOG_CATEGORY(LogNNERuntimeCoreML);
 
 void FNNERuntimeCoreMLModule::StartupModule()
 {
 #ifdef WITH_NNE_RUNTIME_COREML
-	NNERuntimeCoreML = NewObject<UNNERuntimeCoreML>();
+#if WITH_EDITOR
+	NNERuntimeCoreML = NewObject<UNNERuntimeCoreMLCpuGpuNpu>();
+#else // !WITH_EDITOR
+	if (UE::NNERuntimeCoreML::Private::IsNPUAvailable())
+	{
+		NNERuntimeCoreML = NewObject<UNNERuntimeCoreMLCpuGpuNpu>();
+	}
+	else
+	{
+		NNERuntimeCoreML = NewObject<UNNERuntimeCoreMLCpuGpu>();
+	}
+#endif // WITH_EDITOR
+	
 	if (NNERuntimeCoreML.IsValid())
 	{
 		TWeakInterfacePtr<INNERuntime> RuntimeInterface(NNERuntimeCoreML.Get());
-
+		
 		NNERuntimeCoreML->AddToRoot();
 		UE::NNE::RegisterRuntime(RuntimeInterface);
 	}
@@ -26,7 +44,7 @@ void FNNERuntimeCoreMLModule::ShutdownModule()
 	if (NNERuntimeCoreML.IsValid())
 	{
 		TWeakInterfacePtr<INNERuntime> RuntimeInterface(NNERuntimeCoreML.Get());
-
+		
 		UE::NNE::UnregisterRuntime(RuntimeInterface);
 		NNERuntimeCoreML->RemoveFromRoot();
 		NNERuntimeCoreML.Reset();
diff --git a/Engine/Plugins/Experimental/NNERuntimeCoreML/Source/NNERuntimeCoreML/Private/NNERuntimeCoreMLNPUHelper.cpp b/Engine/Plugins/Experimental/NNERuntimeCoreML/Source/NNERuntimeCoreML/Private/NNERuntimeCoreMLNPUHelper.cpp
new file mode 100644
--- /dev/null
+++ b/Engine/Plugins/Experimental/NNERuntimeCoreML/Source/NNERuntimeCoreML/Private/NNERuntimeCoreMLNPUHelper.cpp
@@ -0,0 +1,33 @@
+// Copyright Epic Games, Inc. All Rights Reserved.
+
+#include "NNERuntimeCoreMLNPUHelper.h"
+
+#if defined(__APPLE__)
+#include "CoreMinimal.h"
+#import <Foundation/Foundation.h>
+#import <CoreML/CoreML.h>
+#endif // defined(__APPLE__)
+
+namespace UE::NNERuntimeCoreML::Private
+{
+
+bool IsNPUAvailable()
+{
+#if defined(__APPLE__)
+	SCOPED_AUTORELEASE_POOL;
+	
+	NSArray<id<MLComputeDeviceProtocol>>* Devices = MLAllComputeDevices();
+	
+	for(id<MLComputeDeviceProtocol> Device in Devices)
+	{
+		if ([Device isKindOfClass:[MLNeuralEngineComputeDevice class]]) {
+			return true;
+		}
+	}
+	return false;
+#else // !defined(__APPLE__)
+	return false;
+#endif // defined(__APPLE__)
+}
+
+} // namespace UE::NNERuntimeCoreML::Private
diff --git a/Engine/Plugins/Experimental/NNERuntimeCoreML/Source/NNERuntimeCoreML/Private/NNERuntimeCoreMLNPUHelper.h b/Engine/Plugins/Experimental/NNERuntimeCoreML/Source/NNERuntimeCoreML/Private/NNERuntimeCoreMLNPUHelper.h
new file mode 100644
--- /dev/null
+++ b/Engine/Plugins/Experimental/NNERuntimeCoreML/Source/NNERuntimeCoreML/Private/NNERuntimeCoreMLNPUHelper.h
@@ -0,0 +1,10 @@
+// Copyright Epic Games, Inc. All Rights Reserved.
+
+#pragma once
+
+namespace UE::NNERuntimeCoreML::Private
+{
+
+bool IsNPUAvailable();
+
+} // namespace UE::NNERuntimeCoreML::Private
