[NNE] CoreML runtime: add support for float16/double/int32 for MultyArray input/ouput

Problem: CoreML runtime was only supporting MultiArray of type float
Solution: Add other MultiArray data types

Tests:
Mac editor debug : all tests including added test with new data types

#rnx

diff --git a/Engine/Plugins/Experimental/NNERuntimeCoreML/Source/NNERuntimeCoreML/Private/NNERuntimeCoreMLModel.cpp b/Engine/Plugins/Experimental/NNERuntimeCoreML/Source/NNERuntimeCoreML/Private/NNERuntimeCoreMLModel.cpp
--- a/Engine/Plugins/Experimental/NNERuntimeCoreML/Source/NNERuntimeCoreML/Private/NNERuntimeCoreMLModel.cpp
+++ b/Engine/Plugins/Experimental/NNERuntimeCoreML/Source/NNERuntimeCoreML/Private/NNERuntimeCoreMLModel.cpp
@@ -101,6 +101,32 @@ namespace Detail
 		}
 	}
 
+	ENNETensorDataType GetTypeFromMultiArrayDataType(MLMultiArrayDataType Type)
+	{
+		switch(Type)
+		{
+			case MLMultiArrayDataTypeInt32:		return ENNETensorDataType::Int32;
+			case MLMultiArrayDataTypeFloat16:	return ENNETensorDataType::Half;
+			case MLMultiArrayDataTypeFloat32:	return ENNETensorDataType::Float;
+			case MLMultiArrayDataTypeFloat64:	return ENNETensorDataType::Double;
+			default:
+				return ENNETensorDataType::None;
+		}
+	}
+
+	MLMultiArrayDataType GetMultiArrayDataTypeFromType(ENNETensorDataType Type)
+	{
+		switch(Type)
+		{
+			case ENNETensorDataType::Int32:		return MLMultiArrayDataTypeInt32;
+			case ENNETensorDataType::Half:		return MLMultiArrayDataTypeFloat16;
+			case ENNETensorDataType::Float:		return MLMultiArrayDataTypeFloat32;
+			case ENNETensorDataType::Double:	return MLMultiArrayDataTypeFloat64;
+			default:
+				unimplemented();
+				return MLMultiArrayDataTypeFloat32;
+		}
+	}
 
 	template<typename T> TArray<T, TInlineAllocator<NNE::FTensorShape::MaxRank>> GetShapeDataFromNSArray(const NSArray<NSNumber*>* Shape)
 	{
@@ -147,20 +173,13 @@ namespace Detail
 				FeatureNames.Reset();
 				return false;
 			}
-
-			if (FeatureDescription.multiArrayConstraint.dataType!=MLMultiArrayDataTypeFloat32)
-			{
-				UE_LOG(LogNNERuntimeCoreML, Error, TEXT("Plugin is experimental: Only float32 data type are supported at the moment but %s feature %s is not"), *Context, *FeatureName);
-				TensorDescs.Reset();
-				FeatureNames.Reset();
-				return false;
-			}
 			
 			if (FeatureDescription.multiArrayConstraint.shape.count != 0)
 			{
 				const TArray<int32, TInlineAllocator<NNE::FTensorShape::MaxRank>> SymbolicShapeData = GetShapeDataFromNSArray<int32>(FeatureDescription.multiArrayConstraint.shape);
+				const ENNETensorDataType DataType = GetTypeFromMultiArrayDataType(FeatureDescription.multiArrayConstraint.dataType);
 				const NNE::FSymbolicTensorShape SymbolicShape = NNE::FSymbolicTensorShape::Make(SymbolicShapeData);
-				const NNE::FTensorDesc SymbolicTensor = NNE::FTensorDesc::Make(FeatureName, SymbolicShape, ENNETensorDataType::Float);
+				const NNE::FTensorDesc SymbolicTensor = NNE::FTensorDesc::Make(FeatureName, SymbolicShape, DataType);
 				
 				TensorDescs.Emplace(SymbolicTensor);
 			}
@@ -301,12 +320,13 @@ FModelInstanceCoreMLCpu::ERunSyncStatus FModelInstanceCoreMLBase<InstanceType>::
 		NNE::FTensorShape& InputTensorShape = InputTensorShapes[i];
 		NSMutableArray* Shape = [NSMutableArray arrayWithCapacity:InputTensorShape.Rank()];
 		NSMutableArray* Strides = [NSMutableArray arrayWithCapacity:InputTensorShape.Rank()];
+		MLMultiArrayDataType CoreMLMultiArrayDataType = Detail::GetMultiArrayDataTypeFromType(InputSymbolicTensors[i].GetDataType());
 		
 		Detail::FillNSArrayShapeAndStrideFromNNEShape(InputTensorShape, Shape, Strides);
 
 		MLMultiArray* InputMultiArray = [[[MLMultiArray alloc]  initWithDataPointer:(void *) InInputBindings[i].Data 
 																shape:Shape
-																dataType:MLMultiArrayDataTypeFloat 
+																dataType:CoreMLMultiArrayDataType 
 																strides:Strides 
 																deallocator:^(void * _Nonnull bytes) {}
 																error: &ErrorCreatingInput]
@@ -429,25 +449,27 @@ bool FModelInstanceCoreMLBase<InstanceType>::Init(TConstArrayView64<uint8> Model
 		return false;
 	}
 	
-	CoreMLModelInstance = ModelInstance;
-	
-	// Extract model metadataß
+	// Extract model metadata
 	if (!Detail::GetTensorDescAndFeatureNamesFromMLDescription(
 		InputSymbolicTensors, InputFeatureNames, 
-		CoreMLModelInstance.modelDescription.inputDescriptionsByName, 
+		ModelInstance.modelDescription.inputDescriptionsByName, 
 		true, TEXT("inputs")))
 	{
 		return false;
 	}
 	if (!Detail::GetTensorDescAndFeatureNamesFromMLDescription(
 		OutputSymbolicTensors, OutputFeatureNames, 
-		CoreMLModelInstance.modelDescription.outputDescriptionsByName, 
+		ModelInstance.modelDescription.outputDescriptionsByName, 
 		false, TEXT("outputs")))
 	{
 		return false;
 	}
 
+	// Store instance pointer and prevent it to be garbage collected.
+	CoreMLModelInstance = ModelInstance;
+	
 	[CoreMLModelInstance retain];
+	
 	return true;
 #else // !defined(__APPLE__)
 	return false;
