[NNE] CoreML runtime plugin registration improvement

Problem: CoreML runtime was registering inference interface on editor even if inference was not available.
Solution: Register CoreML as a pure cooking only runtime on non Mac platforms.

Test:
Mac dev editor and standalone: nneqa coreml tests.
Windows debug editor: verify the interface does not expose inference capabilities + smoketest + nneqa all tests.

#rnx

diff --git a/Engine/Plugins/Experimental/NNERuntimeCoreML/Source/NNERuntimeCoreML/Private/NNERuntimeCoreMLModule.cpp b/Engine/Plugins/Experimental/NNERuntimeCoreML/Source/NNERuntimeCoreML/Private/NNERuntimeCoreMLModule.cpp
--- a/Engine/Plugins/Experimental/NNERuntimeCoreML/Source/NNERuntimeCoreML/Private/NNERuntimeCoreMLModule.cpp
+++ b/Engine/Plugins/Experimental/NNERuntimeCoreML/Source/NNERuntimeCoreML/Private/NNERuntimeCoreMLModule.cpp
@@ -12,21 +12,57 @@
 
 DEFINE_LOG_CATEGORY(LogNNERuntimeCoreML);
 
+namespace UE::NNERuntimeCoreML::Private::Details
+{
+
+bool IsInferenceSupported()
+{
+#if defined(__APPLE__)
+	return true;
+#else // !defined(__APPLE__)
+	return false;
+#endif // defined(__APPLE__)
+}
+
+bool IsNPUAvailable()
+{
+#if defined(__APPLE__)
+	SCOPED_AUTORELEASE_POOL;
+	
+	NSArray<id<MLComputeDeviceProtocol>>* Devices = MLAllComputeDevices();
+	
+	for(id<MLComputeDeviceProtocol> Device in Devices)
+	{
+		if ([Device isKindOfClass:[MLNeuralEngineComputeDevice class]]) {
+			return true;
+		}
+	}
+	return false;
+#else // !defined(__APPLE__)
+	return false;
+#endif // defined(__APPLE__)
+}
+
+} // namespace namespace UE::NNERuntimeCoreML::Private::Details
+
 void FNNERuntimeCoreMLModule::StartupModule()
 {
 #ifdef WITH_NNE_RUNTIME_COREML
-#if WITH_EDITOR
-	NNERuntimeCoreML = NewObject<UNNERuntimeCoreMLCpuGpuNpu>();
-#else // !WITH_EDITOR
-	if (UE::NNERuntimeCoreML::Private::IsNPUAvailable())
+	if (UE::NNERuntimeCoreML::Private::Details::IsInferenceSupported())
 	{
-		NNERuntimeCoreML = NewObject<UNNERuntimeCoreMLCpuGpuNpu>();
-	}
+		if (UE::NNERuntimeCoreML::Private::Details::IsNPUAvailable())
+		{
+			NNERuntimeCoreML = NewObject<UNNERuntimeCoreMLCpuGpuNpu>();
+		}
+		else
+		{
+			NNERuntimeCoreML = NewObject<UNNERuntimeCoreMLCpuGpu>();
+		}
+	} 
 	else
 	{
-		NNERuntimeCoreML = NewObject<UNNERuntimeCoreMLCpuGpu>();
+		NNERuntimeCoreML = NewObject<UNNERuntimeCoreML>();
 	}
-#endif // WITH_EDITOR
 	
 	if (NNERuntimeCoreML.IsValid())
 	{
