[NNE] Fix Log declarations

#rnx
#tests NNEModelTest

Cleanup of some duplicate declarations and definitions of the log struct and properly export them so that other modules can access them.

#Synced-CL 49944224

diff --git a/Engine/Plugins/Experimental/NNERuntimeCoreML/Source/NNERuntimeCoreML/Internal/NNERuntimeCoreMLLog.h b/Engine/Plugins/Experimental/NNERuntimeCoreML/Source/NNERuntimeCoreML/Internal/NNERuntimeCoreMLLog.h
--- a/Engine/Plugins/Experimental/NNERuntimeCoreML/Source/NNERuntimeCoreML/Internal/NNERuntimeCoreMLLog.h
+++ b/Engine/Plugins/Experimental/NNERuntimeCoreML/Source/NNERuntimeCoreML/Internal/NNERuntimeCoreMLLog.h
@@ -4,4 +4,4 @@
 
 #include "Logging/LogMacros.h"
 
-DECLARE_LOG_CATEGORY_EXTERN(LogNNERuntimeCoreML, Log, All);
\ No newline at end of file
+NNERUNTIMECOREML_API DECLARE_LOG_CATEGORY_EXTERN(LogNNERuntimeCoreML, Log, All);
\ No newline at end of file
diff --git a/Engine/Plugins/Experimental/NNERuntimeCoreML/Source/NNERuntimeCoreML/Private/NNERuntimeCoreML.cpp b/Engine/Plugins/Experimental/NNERuntimeCoreML/Source/NNERuntimeCoreML/Private/NNERuntimeCoreML.cpp
--- a/Engine/Plugins/Experimental/NNERuntimeCoreML/Source/NNERuntimeCoreML/Private/NNERuntimeCoreML.cpp
+++ b/Engine/Plugins/Experimental/NNERuntimeCoreML/Source/NNERuntimeCoreML/Private/NNERuntimeCoreML.cpp
@@ -7,6 +7,7 @@
 #include "CoreMinimal.h"
 #include "NNE.h"
 #include "NNEModelData.h"
+#include "NNERuntimeCoreMLLog.h"
 #include "NNERuntimeCoreMLModel.h"
 #include "NNERuntimeCoreMLNPUHelper.h"
 #include "Serialization/MemoryWriter.h"
diff --git a/Engine/Plugins/Experimental/NNERuntimeCoreML/Source/NNERuntimeCoreML/Private/NNERuntimeCoreML.h b/Engine/Plugins/Experimental/NNERuntimeCoreML/Source/NNERuntimeCoreML/Private/NNERuntimeCoreML.h
--- a/Engine/Plugins/Experimental/NNERuntimeCoreML/Source/NNERuntimeCoreML/Private/NNERuntimeCoreML.h
+++ b/Engine/Plugins/Experimental/NNERuntimeCoreML/Source/NNERuntimeCoreML/Private/NNERuntimeCoreML.h
@@ -12,8 +12,6 @@
 
 #include "NNERuntimeCoreML.generated.h"
 
-DECLARE_LOG_CATEGORY_EXTERN(LogNNERuntimeCoreML, Log, All);
-
 UCLASS()
 class UNNERuntimeCoreML : public UObject, public INNERuntime
 {
diff --git a/Engine/Plugins/Experimental/NNERuntimeCoreML/Source/NNERuntimeCoreML/Private/NNERuntimeCoreMLModel.cpp b/Engine/Plugins/Experimental/NNERuntimeCoreML/Source/NNERuntimeCoreML/Private/NNERuntimeCoreMLModel.cpp
--- a/Engine/Plugins/Experimental/NNERuntimeCoreML/Source/NNERuntimeCoreML/Private/NNERuntimeCoreMLModel.cpp
+++ b/Engine/Plugins/Experimental/NNERuntimeCoreML/Source/NNERuntimeCoreML/Private/NNERuntimeCoreMLModel.cpp
@@ -11,6 +11,7 @@
 #include "Misc/Paths.h"
 #include "NNEModelData.h"
 #include "NNERuntimeCoreML.h"
+#include "NNERuntimeCoreMLLog.h"
 
 #if defined(__APPLE__)
 #import <Foundation/Foundation.h>
diff --git a/Engine/Plugins/Experimental/NNERuntimeCoreML/Source/NNERuntimeCoreML/Private/NNERuntimeCoreMLModule.cpp b/Engine/Plugins/Experimental/NNERuntimeCoreML/Source/NNERuntimeCoreML/Private/NNERuntimeCoreMLModule.cpp
--- a/Engine/Plugins/Experimental/NNERuntimeCoreML/Source/NNERuntimeCoreML/Private/NNERuntimeCoreMLModule.cpp
+++ b/Engine/Plugins/Experimental/NNERuntimeCoreML/Source/NNERuntimeCoreML/Private/NNERuntimeCoreMLModule.cpp
@@ -3,6 +3,7 @@
 #include "NNERuntimeCoreMLModule.h"
 #include "NNE.h"
 #include "NNERuntimeCoreML.h"
+#include "NNERuntimeCoreMLLog.h"
 #include "NNERuntimeCoreMLNPUHelper.h"
 
 #if defined(__APPLE__)
diff --git a/Engine/Plugins/Experimental/NNERuntimeIREE/Source/NNERuntimeIREE/Internal/NNERuntimeIREELog.h b/Engine/Plugins/Experimental/NNERuntimeIREE/Source/NNERuntimeIREE/Internal/NNERuntimeIREELog.h
--- a/Engine/Plugins/Experimental/NNERuntimeIREE/Source/NNERuntimeIREE/Internal/NNERuntimeIREELog.h
+++ b/Engine/Plugins/Experimental/NNERuntimeIREE/Source/NNERuntimeIREE/Internal/NNERuntimeIREELog.h
@@ -4,4 +4,4 @@
 
 #include "Logging/LogMacros.h"
 
-DECLARE_LOG_CATEGORY_EXTERN(LogNNERuntimeIREE, Log, All);
\ No newline at end of file
+NNERUNTIMEIREE_API DECLARE_LOG_CATEGORY_EXTERN(LogNNERuntimeIREE, Log, All);
\ No newline at end of file
diff --git a/Engine/Plugins/Experimental/NNERuntimeIREE/Source/NNERuntimeIREEEditor/Private/NNERuntimeIREEEditorModule.cpp b/Engine/Plugins/Experimental/NNERuntimeIREE/Source/NNERuntimeIREEEditor/Private/NNERuntimeIREEEditorModule.cpp
--- a/Engine/Plugins/Experimental/NNERuntimeIREE/Source/NNERuntimeIREEEditor/Private/NNERuntimeIREEEditorModule.cpp
+++ b/Engine/Plugins/Experimental/NNERuntimeIREE/Source/NNERuntimeIREEEditor/Private/NNERuntimeIREEEditorModule.cpp
@@ -5,8 +5,6 @@
 
 #define LOCTEXT_NAMESPACE "FNNERuntimeIREEEditorModule"
 
-DEFINE_LOG_CATEGORY(LogNNERuntimeIREE);
-
 void FNNERuntimeIREEEditorModule::StartupModule()
 {
 }
