[NNE] NNERuntimeCoreML  add .mlpackage format support.

Problem:
Currently only the .mlmodel file format is supported.

Solution:
Also support .mlpackage format, which is a directory tree (File package on macOS). Since UE can't directly import directories, we let the user import Manifest.json (part of each .mlpackage) instead. During import, we place the asset in a new UPackage, named after the model, otherwise it would have the name 'Manifest'. Then, similar to how the .mlmodel is handled currently, we serialize the directory tree into the NNEModelData asset and deserialize it back to disk before loading it with Core ML in ModelInstance init.

#tests

#Synced-CL 50647680

diff --git a/Engine/Plugins/Experimental/NNERuntimeCoreML/NNERuntimeCoreML.uplugin b/Engine/Plugins/Experimental/NNERuntimeCoreML/NNERuntimeCoreML.uplugin
--- a/Engine/Plugins/Experimental/NNERuntimeCoreML/NNERuntimeCoreML.uplugin
+++ b/Engine/Plugins/Experimental/NNERuntimeCoreML/NNERuntimeCoreML.uplugin
@@ -25,6 +25,11 @@
 			"Name": "NNERuntimeCoreML",
 			"Type": "RuntimeAndProgram",
 			"LoadingPhase": "Default"
+		},
+		{
+			"Name": "NNERuntimeCoreMLUtils",
+			"Type": "RuntimeAndProgram",
+			"LoadingPhase": "Default"
 		}
 	]
 }
\ No newline at end of file
diff --git a/Engine/Plugins/Experimental/NNERuntimeCoreML/Source/NNERuntimeCoreML/NNERuntimeCoreML.Build.cs b/Engine/Plugins/Experimental/NNERuntimeCoreML/Source/NNERuntimeCoreML/NNERuntimeCoreML.Build.cs
--- a/Engine/Plugins/Experimental/NNERuntimeCoreML/Source/NNERuntimeCoreML/NNERuntimeCoreML.Build.cs
+++ b/Engine/Plugins/Experimental/NNERuntimeCoreML/Source/NNERuntimeCoreML/NNERuntimeCoreML.Build.cs
@@ -15,6 +15,7 @@ public class NNERuntimeCoreML : ModuleRules
 			"Core",
 			"CoreUObject",
 			"NNE",
+			"NNERuntimeCoreMLUtils",
 			"Projects"
 		});
 
diff --git a/Engine/Plugins/Experimental/NNERuntimeCoreML/Source/NNERuntimeCoreML/Private/NNERuntimeCoreML.cpp b/Engine/Plugins/Experimental/NNERuntimeCoreML/Source/NNERuntimeCoreML/Private/NNERuntimeCoreML.cpp
--- a/Engine/Plugins/Experimental/NNERuntimeCoreML/Source/NNERuntimeCoreML/Private/NNERuntimeCoreML.cpp
+++ b/Engine/Plugins/Experimental/NNERuntimeCoreML/Source/NNERuntimeCoreML/Private/NNERuntimeCoreML.cpp
@@ -9,11 +9,11 @@
 #include "NNEModelData.h"
 #include "NNERuntimeCoreMLLog.h"
 #include "NNERuntimeCoreMLModel.h"
+#include "NNERuntimeCoreMLModelData.h"
 #include "NNERuntimeCoreMLNPUHelper.h"
-#include "Serialization/MemoryWriter.h"
 
 FGuid UNNERuntimeCoreML::GUID = FGuid((int32)'C', (int32)'O', (int32)'M', (int32)'L');
-int32 UNNERuntimeCoreML::Version = 0x00000001;
+int32 UNNERuntimeCoreML::Version = 0x00000002;
 
 namespace UE::NNERuntimeCoreML::Private::Details
 {
@@ -21,52 +21,66 @@ namespace UE::NNERuntimeCoreML::Private::Details
 	{
 		check(ModelData != nullptr);
 
-		constexpr int32 GuidSize = sizeof(UNNERuntimeCoreML::GUID);
-		constexpr int32 VersionSize = sizeof(UNNERuntimeCoreML::Version);
 		const TSharedPtr<UE::NNE::FSharedModelData> SharedData = ModelData->GetModelData(RuntimeName);
-
 		if (!SharedData.IsValid())
 		{
 			return CanCreateModelStatus::Fail;
 		}
 
-		TConstArrayView64<uint8> Data = SharedData->GetView();
-
-		if (Data.Num() <= GuidSize + VersionSize)
+		FModelData CoreMLModelData;
+		if (!CoreMLModelData.Load(SharedData->GetView()) || !CoreMLModelData.CheckGUIDAndVersion(UNNERuntimeCoreML::GUID, UNNERuntimeCoreML::Version))
 		{
 			return CanCreateModelStatus::Fail;
 		}
 
-		bool bResult = FGenericPlatformMemory::Memcmp(&(Data[0]), &(UNNERuntimeCoreML::GUID), GuidSize) == 0;
-		bResult &= FGenericPlatformMemory::Memcmp(&(Data[GuidSize]), &(UNNERuntimeCoreML::Version), VersionSize) == 0;
-
-		return bResult ? CanCreateModelStatus::Ok : CanCreateModelStatus::Fail;
+		return CanCreateModelStatus::Ok;
 	}
 } // UE::NNERuntimeCoreML::Private::Details
 
 UNNERuntimeCoreML::ECanCreateModelDataStatus UNNERuntimeCoreML::CanCreateModelData(const FString& FileType, TConstArrayView64<uint8> FileData, const TMap<FString, TConstArrayView64<uint8>>& AdditionalFileData, const FGuid& FileId, const ITargetPlatform* TargetPlatform) const
 {
-	return (!FileData.IsEmpty() && FileType.Compare("mlmodel", ESearchCase::IgnoreCase) == 0) ? ECanCreateModelDataStatus::Ok : ECanCreateModelDataStatus::FailFileIdNotSupported;
+	if (FileData.IsEmpty())
+	{
+		return ECanCreateModelDataStatus::Fail;
+	}
+
+	const bool bIsMLModel = FileType.Compare("mlmodel", ESearchCase::IgnoreCase) == 0;
+	const bool bIsMLPackage = FileType.Compare("mlpackage", ESearchCase::IgnoreCase) == 0;
+	if (!bIsMLModel && !bIsMLPackage)
+	{
+		return ECanCreateModelDataStatus::FailFileIdNotSupported;
+	}
+
+	return ECanCreateModelDataStatus::Ok;
 }
 
 TSharedPtr<UE::NNE::FSharedModelData> UNNERuntimeCoreML::CreateModelData(const FString& FileType, TConstArrayView64<uint8> FileData, const TMap<FString, TConstArrayView64<uint8>>& AdditionalFileData, const FGuid& FileId, const ITargetPlatform* TargetPlatform)
 {
-	using namespace UE::NNERuntimeCoreML::Private;
+	using namespace UE::NNERuntimeCoreML;
 
 	if (CanCreateModelData(FileType, FileData, AdditionalFileData, FileId, TargetPlatform) != ECanCreateModelDataStatus::Ok)
 	{
 		UE_LOG(LogNNERuntimeCoreML, Error, TEXT("Cannot create the model data with id %s (Filetype: %s)"), *FileId.ToString(EGuidFormats::Digits).ToLower(), *FileType);
 		return {};
 	}
-
 	//Note: here the model should be optimized and the related mlmodelc binary blob saved instead of source model
-
+	
+	const bool bIsMLModel = FileType.Compare("mlmodel", ESearchCase::IgnoreCase) == 0;
+	
+	FModelData ModelData
+	{
+		.GUID = UNNERuntimeCoreML::GUID,
+		.Version = UNNERuntimeCoreML::Version,
+		.ModelType = bIsMLModel ? EModelType::MLModel : EModelType::MLPackage,
+		.FileDataView = FileData
+	};
+	
 	TArray64<uint8> Result;
-	FMemoryWriter64 Writer(Result, /*bIsPersitent =*/ true);
-	Writer << UNNERuntimeCoreML::GUID;
-	Writer << UNNERuntimeCoreML::Version;
-
-	Writer.Serialize(const_cast<uint8*>(FileData.GetData()), FileData.Num());
+	if (!ModelData.Store(Result))
+	{
+		UE_LOG(LogNNERuntimeCoreML, Error, TEXT("Model data serialization failed"));
+		return {};
+	}
 
 	return MakeShared<UE::NNE::FSharedModelData>(MakeSharedBufferFromArray(MoveTemp(Result)), 0);
 }
diff --git a/Engine/Plugins/Experimental/NNERuntimeCoreML/Source/NNERuntimeCoreML/Private/NNERuntimeCoreMLModel.cpp b/Engine/Plugins/Experimental/NNERuntimeCoreML/Source/NNERuntimeCoreML/Private/NNERuntimeCoreMLModel.cpp
--- a/Engine/Plugins/Experimental/NNERuntimeCoreML/Source/NNERuntimeCoreML/Private/NNERuntimeCoreMLModel.cpp
+++ b/Engine/Plugins/Experimental/NNERuntimeCoreML/Source/NNERuntimeCoreML/Private/NNERuntimeCoreMLModel.cpp
@@ -12,6 +12,8 @@
 #include "NNEModelData.h"
 #include "NNERuntimeCoreML.h"
 #include "NNERuntimeCoreMLLog.h"
+#include "NNERuntimeCoreMLModelData.h"
+#include "NNERuntimeCoreMLUtils.h"
 
 #if PLATFORM_MAC
 #include "Mac/MacSystemIncludes.h"
@@ -67,15 +69,15 @@ namespace Detail
 		static MLComputeUnits GetComputeUnits() { return MLComputeUnitsCPUAndNeuralEngine; }
 	};
 
-	FString CreateTempDirPath(const FString& BasePath)
+	FString GenerateUniqueTempDirPath(const FString& BasePath)
 	{
-		FString UniqueDirName;
+		FString UniqueDirPath;
 		do
 		{
-			UniqueDirName = FPaths::Combine(BasePath, *FString::Printf(TEXT("CoreMLModel_%s"), *FGuid::NewGuid().ToString()));
-		} while (IFileManager::Get().DirectoryExists(*UniqueDirName));
-		
-		return UniqueDirName;
+			UniqueDirPath = FPaths::Combine(BasePath, *FString::Printf(TEXT("CoreMLModel_%s"), *FGuid::NewGuid().ToString()));
+		} while (IFileManager::Get().DirectoryExists(*UniqueDirPath));
+
+		return UniqueDirPath;
 	}
 
 	void FillNSArrayShapeAndStrideFromNNEShape(const NNE::FTensorShape& TensorShape, NSMutableArray<NSNumber*>* Shape, NSMutableArray<NSNumber*>* Strides)
@@ -405,22 +407,58 @@ FModelInstanceCoreMLCpu::ERunSyncStatus FModelInstanceCoreMLBase<InstanceType>::
 template <class InstanceType>
 bool FModelInstanceCoreMLBase<InstanceType>::Init(TConstArrayView64<uint8> ModelData)
 {
+	SCOPED_NAMED_EVENT_TEXT("FModelInstanceCoreMLBase::Init", FColor::Magenta);
+
 #if defined(__APPLE__)
+
 	SCOPED_AUTORELEASE_POOL;
-	
-	// Find model memory blob and save it to disk.
-	int32 GuidSize = sizeof(UNNERuntimeCoreML::GUID);
-	int32 VersionSize = sizeof(UNNERuntimeCoreML::Version);
-	int32 GuidAndVersionSize = GuidSize+VersionSize;
-	TConstArrayView<uint8> ModelBuffer = { &(ModelData.GetData()[GuidAndVersionSize]), ModelData.Num() - GuidAndVersionSize };
+
+	using namespace UE::NNERuntimeCoreML;
+
+	FModelData CoreMLModelData;
+	if (!CoreMLModelData.Load(ModelData))
+	{
+		UE_LOG(LogNNERuntimeCoreML, Error, TEXT("Failed to read model data."));
+		return false;
+	}
+
+	// Save file data to disk
 	FString ProjIntermediateDir = FPaths::ConvertRelativePathToFull(FPaths::ProjectIntermediateDir());
-	FString TempDirForModel = Detail::CreateTempDirPath(ProjIntermediateDir);
-	FString Filepath = FPaths::Combine(TempDirForModel, TEXT("CoreMLModel.mlmodel"));
+	FString TempDirForModel = Detail::GenerateUniqueTempDirPath(ProjIntermediateDir);
+	FString Filepath;
+	if (CoreMLModelData.ModelType == EModelType::MLModel)
+	{
+		Filepath = FPaths::Combine(TempDirForModel, TEXT("CoreMLModel.mlmodel"));
+		if (!FFileHelper::SaveArrayToFile(CoreMLModelData.FileDataView, *Filepath))
+		{
+			IFileManager::Get().DeleteDirectory(*TempDirForModel, false, true);
+			UE_LOG(LogNNERuntimeCoreML, Error, TEXT("Could not write model to disk at path %s."), *Filepath);
+			return false;
+		}
+	}
+	else if (CoreMLModelData.ModelType == EModelType::MLPackage)
+	{
+		Filepath = FPaths::Combine(TempDirForModel, TEXT("CoreMLModel.mlpackage"));
+		if (FPaths::DirectoryExists(Filepath))
+		{
+			IPlatformFile& PlatformFile = FPlatformFileManager::Get().GetPlatformFile();
+			if (!PlatformFile.DeleteDirectoryRecursively(*Filepath))
+			{
+				UE_LOG(LogNNERuntimeCoreML, Error, TEXT("Failed to delete directory %s"), *Filepath);
+				return false;
+			}
+		}
 
-	if (!FFileHelper::SaveArrayToFile(ModelBuffer, *Filepath))
+		if (!SaveArrayToDirectory(CoreMLModelData.FileDataView, *Filepath))
+		{
+			IFileManager::Get().DeleteDirectory(*TempDirForModel, false, true);
+			UE_LOG(LogNNERuntimeCoreML, Error, TEXT("Could not write model to disk at path %s."), *Filepath);
+			return false;
+		}
+	}
+	else
 	{
-		IFileManager::Get().DeleteDirectory(*TempDirForModel, false, true);
-		UE_LOG(LogNNERuntimeCoreML, Error, TEXT("Could not write model to disk at path %s."), *Filepath);
+		checkf(false, TEXT("Unexpected ModelType %i."), static_cast<int32>(CoreMLModelData.ModelType));
 		return false;
 	}
 	
@@ -437,7 +475,7 @@ bool FModelInstanceCoreMLBase<InstanceType>::Init(TConstArrayView64<uint8> Model
 	}
 	if (!IFileManager::Get().DeleteDirectory(*TempDirForModel, false, true))
 	{
-		UE_LOG(LogNNERuntimeCoreML, Error, TEXT("Could not cleanup model temp directy %s."), *TempDirForModel);
+		UE_LOG(LogNNERuntimeCoreML, Error, TEXT("Could not cleanup model temp directory %s."), *TempDirForModel);
 	}
 
 	// Create the model instance from the compiled model.
diff --git a/Engine/Plugins/Experimental/NNERuntimeCoreML/Source/NNERuntimeCoreML/Private/NNERuntimeCoreMLModelData.cpp b/Engine/Plugins/Experimental/NNERuntimeCoreML/Source/NNERuntimeCoreML/Private/NNERuntimeCoreMLModelData.cpp
new file mode 100644
--- /dev/null
+++ b/Engine/Plugins/Experimental/NNERuntimeCoreML/Source/NNERuntimeCoreML/Private/NNERuntimeCoreMLModelData.cpp
@@ -0,0 +1,43 @@
+// Copyright Epic Games, Inc. All Rights Reserved.
+
+#include "NNERuntimeCoreMLModelData.h"
+
+#include "Serialization/MemoryReader.h"
+#include "Serialization/MemoryWriter.h"
+
+namespace UE::NNERuntimeCoreML
+{
+	bool FModelData::Load(TConstArrayView64<uint8> InData)
+	{
+		FMemoryReaderView Reader(InData, /*bIsPersistent*/true);
+		Reader << GUID;
+		Reader << Version;
+		Reader << ModelType;
+
+		const int64 Offset = Reader.Tell();
+		if (Reader.IsError() || Offset > InData.Num())
+		{
+			return false;
+		}
+
+		FileDataView = InData.Slice(Offset, InData.Num() - Offset);
+		
+		return true;
+	}
+
+	bool FModelData::Store(TArray64<uint8>& OutData)
+	{
+		FMemoryWriter64 Writer(OutData, /*bIsPersistent*/true);
+		Writer << GUID;
+		Writer << Version;
+		Writer << ModelType;
+		Writer.Serialize((void*)FileDataView.GetData(), FileDataView.Num());
+		
+		return !Writer.IsError();
+	}
+
+	bool FModelData::CheckGUIDAndVersion(const FGuid& ReferenceGUID, int32 ReferenceVersion) const
+	{
+		return ReferenceGUID == GUID && ReferenceVersion == Version;
+	}
+}
\ No newline at end of file
diff --git a/Engine/Plugins/Experimental/NNERuntimeCoreML/Source/NNERuntimeCoreML/Private/NNERuntimeCoreMLModelData.h b/Engine/Plugins/Experimental/NNERuntimeCoreML/Source/NNERuntimeCoreML/Private/NNERuntimeCoreMLModelData.h
new file mode 100644
--- /dev/null
+++ b/Engine/Plugins/Experimental/NNERuntimeCoreML/Source/NNERuntimeCoreML/Private/NNERuntimeCoreMLModelData.h
@@ -0,0 +1,34 @@
+// Copyright Epic Games, Inc. All Rights Reserved.
+
+#pragma once
+
+#include "Containers/ArrayView.h"
+#include "Misc/Guid.h"
+
+namespace UE::NNERuntimeCoreML
+{
+	enum class EModelType : uint8
+	{
+		MLModel,
+		MLPackage
+	};
+
+	class FModelData
+	{
+	public:
+		FGuid GUID;
+
+		int32 Version = 0;
+
+		EModelType ModelType = EModelType::MLModel;
+		
+		// Non-owning view into file data. On saving, copy of file data is written. When loading, points to source passed to Load().
+		TConstArrayView64<uint8> FileDataView;
+
+		bool Load(TConstArrayView64<uint8> InData);
+
+		bool Store(TArray64<uint8>& OutData);
+
+		bool CheckGUIDAndVersion(const FGuid& ReferenceGUID, int32 ReferenceVersion) const;
+	};
+}
\ No newline at end of file
diff --git a/Engine/Plugins/Experimental/NNERuntimeCoreML/Source/NNERuntimeCoreMLEditor/NNERuntimeCoreMLEditor.Build.cs b/Engine/Plugins/Experimental/NNERuntimeCoreML/Source/NNERuntimeCoreMLEditor/NNERuntimeCoreMLEditor.Build.cs
--- a/Engine/Plugins/Experimental/NNERuntimeCoreML/Source/NNERuntimeCoreMLEditor/NNERuntimeCoreMLEditor.Build.cs
+++ b/Engine/Plugins/Experimental/NNERuntimeCoreML/Source/NNERuntimeCoreMLEditor/NNERuntimeCoreMLEditor.Build.cs
@@ -18,6 +18,7 @@ public class NNERuntimeCoreMLEditor : ModuleRules
 				"MainFrame",
 				"NNE",
 				"NNERuntimeCoreML",
+				"NNERuntimeCoreMLUtils",
 				"Slate",
 				"SlateCore",
 				"UnrealEd"
diff --git a/Engine/Plugins/Experimental/NNERuntimeCoreML/Source/NNERuntimeCoreMLEditor/Private/NNERuntimeCoreMLModelDataFactory.cpp b/Engine/Plugins/Experimental/NNERuntimeCoreML/Source/NNERuntimeCoreMLEditor/Private/NNERuntimeCoreMLModelDataFactory.cpp
--- a/Engine/Plugins/Experimental/NNERuntimeCoreML/Source/NNERuntimeCoreMLEditor/Private/NNERuntimeCoreMLModelDataFactory.cpp
+++ b/Engine/Plugins/Experimental/NNERuntimeCoreML/Source/NNERuntimeCoreMLEditor/Private/NNERuntimeCoreMLModelDataFactory.cpp
@@ -3,42 +3,126 @@
 #include "NNERuntimeCoreMLModelDataFactory.h"
 
 #include "Editor.h"
-#include "NNEModelData.h"
+#include "Misc/FileHelper.h"
+#include "Misc/PackageName.h"
 #include "Misc/Paths.h"
+#include "NNEModelData.h"
+#include "NNERuntimeCoreMLLog.h"
+#include "NNERuntimeCoreMLUtils.h"
 #include "Subsystems/ImportSubsystem.h"
 
-
 UNNERuntimeCoreMLModelDataFactory::UNNERuntimeCoreMLModelDataFactory(const FObjectInitializer& ObjectInitializer) : UFactory(ObjectInitializer)
 {
 	bCreateNew = false;
 	bEditorImport = true;
 	SupportedClass = UNNEModelData::StaticClass();
-	ImportPriority = DefaultImportPriority;
+	ImportPriority = DefaultImportPriority + 100;
 	Formats.Add(TEXT("mlmodel;Core ML model format"));
+	Formats.Add(TEXT("json;Core ML ML Program format"));
 }
 
-UObject* UNNERuntimeCoreMLModelDataFactory::FactoryCreateBinary(UClass* InClass, UObject* InParent, FName Name, EObjectFlags Flags, UObject* Context, const TCHAR * Type, const uint8 *& Buffer, const uint8 * BufferEnd, FFeedbackContext* Warn)
+UObject* UNNERuntimeCoreMLModelDataFactory::FactoryCreateFile(UClass* InClass, UObject* InParent, FName InName, EObjectFlags Flags, const FString& Filename, const TCHAR* Parms, FFeedbackContext* Warn, bool& bOutOperationCanceled)
 {
-	GEditor->GetEditorSubsystem<UImportSubsystem>()->BroadcastAssetPreImport(this, InClass, InParent, Name, Type);
+	SCOPED_NAMED_EVENT_TEXT("UNNERuntimeCoreMLModelDataFactory::FactoryCreateFile", FColor::Magenta);
+	
+	using namespace UE::NNERuntimeCoreML;
+
+	const FString Extension = FPaths::GetExtension(Filename);
+	check(Extension == TEXT("mlmodel") || Extension == TEXT("json"));
+
+	const bool bIsMLModel = Extension == TEXT("mlmodel");
+	const FString Type = bIsMLModel ? TEXT("mlmodel") : TEXT("mlpackage");
+
+	// Determine the correct asset name
+	FString AssetName = InName.ToString();
+	UObject* Parent = InParent;
 
-	if (!Type || !Buffer || !BufferEnd || BufferEnd - Buffer <= 0)
+	const FString ParentPath = FPaths::GetPath(Filename);
+
+	if (!bIsMLModel)
 	{
-		GEditor->GetEditorSubsystem<UImportSubsystem>()->BroadcastAssetPostImport(this, nullptr);
-		return nullptr;
+		// Extract name from "ModelName.mlpackage"
+		AssetName = FPaths::GetBaseFilename(ParentPath);
+
+		// Create new package with correct name
+		const FString PackageName = FPackageName::GetLongPackagePath(InParent->GetName());
+		const FString NewPackageName = PackageName / AssetName;
+		
+		UPackage* Pkg = CreatePackage(*NewPackageName);
+		if(!ensure(Pkg))
+		{
+			UE_LOG(LogNNERuntimeCoreML, Error, TEXT("Failed to create package '%s' for asset '%s'"), *NewPackageName, *AssetName);
+			GEditor->GetEditorSubsystem<UImportSubsystem>()->BroadcastAssetPostImport(this, nullptr);
+			return nullptr;
+		}
+
+		Pkg->FullyLoad();
+		Parent = Pkg;
+
+		UObject* ExistingObject = StaticFindObject(UObject::StaticClass(), Pkg, *AssetName);
+		if (ExistingObject)
+		{
+			// Overwrite existing asset if it's the same type, otherwise cancel import to avoid data loss
+			if (DoesSupportClass(ExistingObject->GetClass()))
+			{
+				UE_LOG(LogNNERuntimeCoreML, Log, TEXT("Replaced existing Core ML asset '%s' in package '%s'."), *AssetName, *NewPackageName);
+			}
+			else
+			{
+				UE_LOG(LogNNERuntimeCoreML, Warning, TEXT("An asset named '%s' already exists in package '%s', will not import."), *AssetName, *NewPackageName);
+
+				bOutOperationCanceled = true;
+
+				GEditor->GetEditorSubsystem<UImportSubsystem>()->BroadcastAssetPostImport(this, nullptr);
+				return nullptr;
+			}
+		}
 	}
 
-	TConstArrayView<uint8> BufferView = MakeArrayView(Buffer, BufferEnd - Buffer);
-	UNNEModelData* ModelData = NewObject<UNNEModelData>(InParent, InClass, Name, Flags);
+	GEditor->GetEditorSubsystem<UImportSubsystem>()->BroadcastAssetPreImport(this, InClass, Parent, FName(*AssetName), *Type);
+	
+	TArray64<uint8> Data;
+	if (bIsMLModel)
+	{
+		if (!FFileHelper::LoadFileToArray(Data, *Filename))
+		{
+			UE_LOG(LogNNERuntimeCoreML, Error, TEXT("Failed to load file '%s' to array"), *Filename);
+			GEditor->GetEditorSubsystem<UImportSubsystem>()->BroadcastAssetPostImport(this, nullptr);
+			return nullptr;
+		}
+	}
+	else
+	{
+		if (!LoadDirectoryToArray(Data, *ParentPath))
+		{
+			UE_LOG(LogNNERuntimeCoreML, Error, TEXT("Failed to load folder '%s' to array"), *ParentPath);
+			GEditor->GetEditorSubsystem<UImportSubsystem>()->BroadcastAssetPostImport(this, nullptr);
+			return nullptr;
+		}
+	}
+
+	UNNEModelData* ModelData = NewObject<UNNEModelData>(Parent, InClass, FName(*AssetName), Flags);
 	check(ModelData)
 
-	ModelData->Init(Type, BufferView);
+	ModelData->Init(Type, Data);
 
 	GEditor->GetEditorSubsystem<UImportSubsystem>()->BroadcastAssetPostImport(this, ModelData);
-
 	return ModelData;
 }
 
 bool UNNERuntimeCoreMLModelDataFactory::FactoryCanImport(const FString & Filename)
 {
-	return Filename.EndsWith(FString(TEXT("mlmodel")));
+	if (Filename.EndsWith(TEXT(".mlmodel")))
+	{
+		return true;
+	}
+	
+	const FString CleanFilename = FPaths::GetCleanFilename(Filename);
+	const FString ParentPath = FPaths::GetPath(Filename);
+	if (CleanFilename == TEXT("Manifest.json") && ParentPath.EndsWith(TEXT(".mlpackage")))
+	{
+		return true;
+	}
+
+	return false;
 }
diff --git a/Engine/Plugins/Experimental/NNERuntimeCoreML/Source/NNERuntimeCoreMLEditor/Public/NNERuntimeCoreMLModelDataFactory.h b/Engine/Plugins/Experimental/NNERuntimeCoreML/Source/NNERuntimeCoreMLEditor/Public/NNERuntimeCoreMLModelDataFactory.h
--- a/Engine/Plugins/Experimental/NNERuntimeCoreML/Source/NNERuntimeCoreMLEditor/Public/NNERuntimeCoreMLModelDataFactory.h
+++ b/Engine/Plugins/Experimental/NNERuntimeCoreML/Source/NNERuntimeCoreMLEditor/Public/NNERuntimeCoreMLModelDataFactory.h
@@ -17,7 +17,7 @@ public:
 
 public:
 	//~ Begin UFactory Interface
-	virtual UObject* FactoryCreateBinary(UClass* Class, UObject* InParent, FName Name, EObjectFlags Flags, UObject* Context, const TCHAR* Type, const uint8*& Buffer, const uint8* BufferEnd, FFeedbackContext* Warn) override;
+	virtual UObject* FactoryCreateFile(UClass* InClass, UObject* InParent, FName InName, EObjectFlags Flags, const FString& Filename, const TCHAR* Parms, FFeedbackContext* Warn, bool& bOutOperationCanceled) override;
 	virtual bool FactoryCanImport(const FString& Filename) override;
 	//~ End UFactory Interface
-};
\ No newline at end of file
+};
diff --git a/Engine/Plugins/Experimental/NNERuntimeCoreML/Source/NNERuntimeCoreMLUtils/Internal/NNERuntimeCoreMLUtils.h b/Engine/Plugins/Experimental/NNERuntimeCoreML/Source/NNERuntimeCoreMLUtils/Internal/NNERuntimeCoreMLUtils.h
new file mode 100644
--- /dev/null
+++ b/Engine/Plugins/Experimental/NNERuntimeCoreML/Source/NNERuntimeCoreMLUtils/Internal/NNERuntimeCoreMLUtils.h
@@ -0,0 +1,12 @@
+// Copyright Epic Games, Inc. All Rights Reserved.
+
+#pragma once
+
+#include "Containers/ContainersFwd.h"
+
+namespace UE::NNERuntimeCoreML
+{
+	NNERUNTIMECOREMLUTILS_API bool LoadDirectoryToArray(TArray64<uint8>& Result, const FString& Path);
+
+	NNERUNTIMECOREMLUTILS_API bool SaveArrayToDirectory(TConstArrayView64<uint8> Data, const TCHAR* Path);
+};
\ No newline at end of file
diff --git a/Engine/Plugins/Experimental/NNERuntimeCoreML/Source/NNERuntimeCoreMLUtils/Internal/NNERuntimeCoreMLUtilsLog.h b/Engine/Plugins/Experimental/NNERuntimeCoreML/Source/NNERuntimeCoreMLUtils/Internal/NNERuntimeCoreMLUtilsLog.h
new file mode 100644
--- /dev/null
+++ b/Engine/Plugins/Experimental/NNERuntimeCoreML/Source/NNERuntimeCoreMLUtils/Internal/NNERuntimeCoreMLUtilsLog.h
@@ -0,0 +1,7 @@
+// Copyright Epic Games, Inc. All Rights Reserved.
+
+#pragma once
+
+#include "Logging/LogMacros.h"
+
+DECLARE_LOG_CATEGORY_EXTERN(LogNNERuntimeCoreMLUtils, Log, All);
\ No newline at end of file
diff --git a/Engine/Plugins/Experimental/NNERuntimeCoreML/Source/NNERuntimeCoreMLUtils/NNERuntimeCoreMLUtils.Build.cs b/Engine/Plugins/Experimental/NNERuntimeCoreML/Source/NNERuntimeCoreMLUtils/NNERuntimeCoreMLUtils.Build.cs
new file mode 100644
--- /dev/null
+++ b/Engine/Plugins/Experimental/NNERuntimeCoreML/Source/NNERuntimeCoreMLUtils/NNERuntimeCoreMLUtils.Build.cs
@@ -0,0 +1,19 @@
+// Copyright Epic Games, Inc. All Rights Reserved.
+
+using UnrealBuildTool;
+
+public class NNERuntimeCoreMLUtils : ModuleRules
+{
+	public NNERuntimeCoreMLUtils(ReadOnlyTargetRules Target) : base(Target)
+	{
+		PCHUsage = PCHUsageMode.UseExplicitOrSharedPCHs;
+
+		PrivateDependencyModuleNames.AddRange
+		(
+			new string[]
+			{
+				"Core"
+			}
+		);
+	}
+}
diff --git a/Engine/Plugins/Experimental/NNERuntimeCoreML/Source/NNERuntimeCoreMLUtils/Private/NNERuntimeCoreMLUtils.cpp b/Engine/Plugins/Experimental/NNERuntimeCoreML/Source/NNERuntimeCoreMLUtils/Private/NNERuntimeCoreMLUtils.cpp
new file mode 100644
--- /dev/null
+++ b/Engine/Plugins/Experimental/NNERuntimeCoreML/Source/NNERuntimeCoreMLUtils/Private/NNERuntimeCoreMLUtils.cpp
@@ -0,0 +1,212 @@
+// Copyright Epic Games, Inc. All Rights Reserved.
+
+#include "NNERuntimeCoreMLUtils.h"
+
+#include "Containers/Array.h"
+#include "Containers/ArrayView.h"
+#include "GenericPlatform/GenericPlatformFile.h"
+#include "HAL/FileManager.h"
+#include "HAL/PlatformFileManager.h"
+#include "Math/Color.h"
+#include "Misc/FileHelper.h"
+#include "Misc/Paths.h"
+#include "NNERuntimeCoreMLUtilsLog.h"
+#include "Serialization/Archive.h"
+#include "Serialization/MemoryReader.h"
+#include "Serialization/MemoryWriter.h"
+#include "Templates/UnrealTemplate.h"
+
+namespace UE::NNERuntimeCoreML
+{
+
+namespace Private
+{
+
+constexpr int32 FolderSerializationVersion = 1;
+
+struct FSerializedFileEntry
+{
+	FString RelativePath;
+	int64 FileSize;
+
+	friend FArchive& operator<<(FArchive& Ar, FSerializedFileEntry& Entry)
+	{
+		Ar << Entry.RelativePath;
+		Ar << Entry.FileSize;
+		return Ar;
+	}
+};
+
+struct FSerializedDirectoryHeader
+{
+	int32 Version = FolderSerializationVersion;
+	TArray<FSerializedFileEntry> Files;
+
+	friend FArchive& operator<<(FArchive& Ar, FSerializedDirectoryHeader& Header)
+	{
+		Ar << Header.Version;
+		Ar << Header.Files;
+		return Ar;
+	}
+};
+
+struct FFileInfo
+{
+	FString FullPath;
+	FSerializedFileEntry FileEntry;
+};
+
+} // namespace Private
+
+bool LoadDirectoryToArray(TArray64<uint8>& Result, const FString& Path)
+{
+	SCOPED_NAMED_EVENT_TEXT("CoreML::LoadDirectoryToArray", FColor::Magenta);
+	
+	using namespace Private;
+
+	TArray<FFileInfo> Files;
+	int64 TotalContentSize = 0;
+
+	FString DirectoryPath = Path;
+	FPaths::NormalizeDirectoryName(DirectoryPath);
+	if (!DirectoryPath.EndsWith(TEXT("/")))
+	{
+		DirectoryPath.Append(TEXT("/"));
+	}
+
+	const bool bSuccess = IFileManager::Get().IterateDirectoryStatRecursively(*DirectoryPath,
+		[&DirectoryPath, &Files, &TotalContentSize](const TCHAR* FilenameOrDirectory, const FFileStatData& StatData)
+		{
+			if (!StatData.bIsDirectory)
+			{
+				if (StatData.FileSize < 0)
+				{
+					UE_LOG(LogNNERuntimeCoreMLUtils, Warning, TEXT("File %s has unknown size (FileSize is %lld)"), FilenameOrDirectory, StatData.FileSize);
+					return false;
+				}
+
+				FString RelativePath(FilenameOrDirectory);
+				if (!FPaths::MakePathRelativeTo(RelativePath, *DirectoryPath))
+				{
+					UE_LOG(LogNNERuntimeCoreMLUtils, Warning, TEXT("Failed to convert %s to a relative path of %s"), FilenameOrDirectory, *DirectoryPath);
+					return false;
+				}
+
+				TotalContentSize += StatData.FileSize;
+
+				Files.Add(
+				{
+					.FullPath = FString(FilenameOrDirectory),
+					.FileEntry = { .RelativePath = MoveTemp(RelativePath), .FileSize = StatData.FileSize }
+				});
+			}
+			return true;
+		});
+
+	if (!bSuccess)
+	{
+		return false;
+	}
+
+	FSerializedDirectoryHeader Header;
+	Header.Version = FolderSerializationVersion;
+	Header.Files.Reserve(Files.Num());
+	for (const FFileInfo& Info : Files)
+	{
+		Header.Files.Add(Info.FileEntry);
+	}
+
+	TArray<uint8> HeaderBuffer;
+	FMemoryWriter HeaderWriter(HeaderBuffer, /*bIsPersistent*/true);
+	HeaderWriter << Header;
+
+	const int64 TotalSize = HeaderBuffer.Num() + TotalContentSize;
+	Result.SetNumUninitialized(TotalSize);
+
+	// Copy header directly
+	FMemory::Memcpy(Result.GetData(), HeaderBuffer.GetData(), HeaderBuffer.Num());
+
+	int64 WriteOffset = HeaderBuffer.Num();
+
+	for (const FFileInfo& Info : Files)
+	{
+		SCOPED_NAMED_EVENT_TEXT("CoreML::LoadDirectoryToArray::Loop_file_read", FColor::Magenta);
+		TUniquePtr<IFileHandle> FileHandle(FPlatformFileManager::Get().GetPlatformFile().OpenRead(*Info.FullPath));
+		if (!FileHandle || !FileHandle->Read(Result.GetData() + WriteOffset, Info.FileEntry.FileSize))
+		{
+			UE_LOG(LogNNERuntimeCoreMLUtils, Warning, TEXT("Failed to read from %s"), *Info.FullPath);
+			return false;
+		}
+		WriteOffset += Info.FileEntry.FileSize;
+	}
+ 
+	ensureMsgf(Result.Num() == TotalSize, TEXT("Serialization used different amount than estimated: %lld (estimated %lld)."), Result.Num(), TotalSize);
+	
+	return !HeaderWriter.IsError();
+}
+
+bool SaveArrayToDirectory(TConstArrayView64<uint8> Data, const TCHAR* Path)
+{
+	SCOPED_NAMED_EVENT_TEXT("CoreML::SaveArrayToDirectory", FColor::Magenta);
+
+	using namespace Private;
+
+	FString BasePath = Path;
+	FPaths::NormalizeDirectoryName(BasePath);
+
+	if (!IFileManager::Get().MakeDirectory(*BasePath, /*Tree=*/true))
+	{
+		UE_LOG(LogNNERuntimeCoreMLUtils, Warning, TEXT("Failed to create directory %s"), *BasePath);
+		return false;
+	}
+
+	FMemoryReaderView Reader(Data);
+	FSerializedDirectoryHeader Header;
+	Reader << Header;
+
+	if (Header.Version != FolderSerializationVersion)
+	{
+		UE_LOG(LogNNERuntimeCoreMLUtils, Warning, TEXT("Unexpected version %i encountered"), Header.Version);
+		return false;
+	}
+
+	FString BasePathForCheck = BasePath;
+	if (!BasePathForCheck.EndsWith(TEXT("/")))
+	{
+		BasePathForCheck.Append(TEXT("/"));
+	}
+
+	for (const FSerializedFileEntry& Entry : Header.Files)
+	{
+		const FString AbsolutePath = FPaths::Combine(Path, Entry.RelativePath);
+		FString NormalizedPath = AbsolutePath;
+		FPaths::CollapseRelativeDirectories(NormalizedPath);
+
+		if (!NormalizedPath.StartsWith(BasePathForCheck))
+		{
+			UE_LOG(LogNNERuntimeCoreMLUtils, Warning, TEXT("Path traversal detected: %s resolved outside base directory"), *Entry.RelativePath);
+			return false;
+		}
+		
+		int64 Start = Reader.Tell();
+		int64 End = Start + Entry.FileSize;
+		if (Entry.FileSize < 0 || End > Data.Num())
+		{
+			UE_LOG(LogNNERuntimeCoreMLUtils, Warning, TEXT("Corrupted File Size. Read %lld at location %lld but data only has total size %lld."), Entry.FileSize, Start, Data.Num());
+			return false;
+		}
+
+		TConstArrayView64<uint8> Buffer(Data.GetData() + Start, Entry.FileSize);
+		if (!FFileHelper::SaveArrayToFile(Buffer, *NormalizedPath))
+		{
+			UE_LOG(LogNNERuntimeCoreMLUtils, Warning, TEXT("Save to file %s failed"), *NormalizedPath);
+			return false;
+		}
+
+		Reader.Seek(End);
+	}
+
+	return !Reader.IsError();
+}
+
+} // namespace UE::NNERuntimeCoreML
\ No newline at end of file
diff --git a/Engine/Plugins/Experimental/NNERuntimeCoreML/Source/NNERuntimeCoreMLUtils/Private/NNERuntimeCoreMLUtilsModule.cpp b/Engine/Plugins/Experimental/NNERuntimeCoreML/Source/NNERuntimeCoreMLUtils/Private/NNERuntimeCoreMLUtilsModule.cpp
new file mode 100644
--- /dev/null
+++ b/Engine/Plugins/Experimental/NNERuntimeCoreML/Source/NNERuntimeCoreMLUtils/Private/NNERuntimeCoreMLUtilsModule.cpp
@@ -0,0 +1,16 @@
+// Copyright Epic Games, Inc. All Rights Reserved.
+
+#include "Modules/ModuleManager.h"
+#include "NNERuntimeCoreMLUtilsLog.h"
+
+DEFINE_LOG_CATEGORY(LogNNERuntimeCoreMLUtils);
+
+class FNNERuntimeCoreMLUtilsModule : public IModuleInterface
+{
+public:
+	/** IModuleInterface implementation */
+	virtual void StartupModule() override {}
+	virtual void ShutdownModule() override {}
+};
+
+IMPLEMENT_MODULE(FNNERuntimeCoreMLUtilsModule, NNERuntimeCoreMLUtils)
\ No newline at end of file
