[NNERuntimeCoreML] Fix output copy to use logical size from MLMultiArray shape

CoreML's getBytesWithHandler reports the stride-based allocation size, which may be larger than the logical data due to internal alignment padding (e.g. innermost dimension padded to a multiple of 8). This causes the safety check if (SizeInBytes >= size) to fail for correctly-sized output buffers, since size includes padding the caller doesn't account for.

This change computes the logical data size from OutputMultiArray.count (product of shape dimensions) multiplied by the element byte size, and uses that instead of the raw size for both the guard check and the copy length.

#rnx
#Synced-CL 51630826

diff --git a/Engine/Plugins/Experimental/NNERuntimeCoreML/Source/NNERuntimeCoreML/Private/NNERuntimeCoreMLModel.cpp b/Engine/Plugins/Experimental/NNERuntimeCoreML/Source/NNERuntimeCoreML/Private/NNERuntimeCoreMLModel.cpp
--- a/Engine/Plugins/Experimental/NNERuntimeCoreML/Source/NNERuntimeCoreML/Private/NNERuntimeCoreMLModel.cpp
+++ b/Engine/Plugins/Experimental/NNERuntimeCoreML/Source/NNERuntimeCoreML/Private/NNERuntimeCoreMLModel.cpp
@@ -386,15 +386,17 @@ FModelInstanceCoreMLCpu::ERunSyncStatus FModelInstanceCoreMLBase<InstanceType>::
 		// Copy the memory buffers of the output features to user provided bindings
 		if (!InOutputBindings.IsEmpty() && InOutputBindings[i].Data)
 		{
-			void (^copyToOutput)(const void*, NSInteger) =
-				^(const void* bytes, NSInteger size) {
-					if (size > 0 && InOutputBindings[i].SizeInBytes >= size)
-					{
-						FMemory::Memcpy(InOutputBindings[i].Data, bytes, size);
-					}
-				};
+			const int32 ElementSize = NNE::GetTensorDataTypeSizeInBytes(Detail::GetTypeFromMultiArrayDataType(OutputMultiArray.dataType));
 			
-			[OutputMultiArray getBytesWithHandler:copyToOutput];
+			// OutputMultiArray.count is the product of shape dimensions (logical element count, excludes stride padding)
+			const uint64 LogicalSizeInBytes = static_cast<uint64>(OutputMultiArray.count) * ElementSize;
+
+			if (InOutputBindings[i].SizeInBytes >= LogicalSizeInBytes)
+			{
+				[OutputMultiArray getBytesWithHandler:^(const void* bytes, NSInteger size) {
+					FMemory::Memcpy(InOutputBindings[i].Data, bytes, LogicalSizeInBytes);
+				}];
+			}
 		}
 	}
 	
